Merging CL 2527648, 2527736

diff --git a/Engine/Source/Runtime/JsonUtilities/Public/JsonObjectConverter.h b/Engine/Source/Runtime/JsonUtilities/Public/JsonObjectConverter.h
--- a/Engine/Source/Runtime/JsonUtilities/Public/JsonObjectConverter.h
+++ b/Engine/Source/Runtime/JsonUtilities/Public/JsonObjectConverter.h
@@ -167,9 +167,9 @@ public: // JSON -> UStruct
 	* @param JsonString String containing JSON formatted data.
 	* @param OutStructArray The UStruct array to copy in to
 	* @param CheckFlags Only convert properties that match at least one of these flags. If 0 check all properties.
-	* @param SkipFlags Skip properties that match any of these flags
+	* @param SkipFlags Skip properties that match any of these flags.
 	*
-	* @return False if any properties matched but failed to deserialize
+	* @return False if any properties matched but failed to deserialize.
 	*/
 	template<typename OutStructType>
 	static bool JsonArrayStringToUStruct(const FString& JsonString, TArray<OutStructType>* OutStructArray, int64 CheckFlags, int64 SkipFlags)
@@ -181,18 +181,39 @@ public: // JSON -> UStruct
 			UE_LOG(LogJson, Warning, TEXT("JsonArrayStringToUStruct - Unable to parse. json=[%s]"), *JsonString);
 			return false;
 		}
+		if (!JsonArrayToUStruct(JsonArray, OutStructArray, CheckFlags, SkipFlags))
+		{
+			UE_LOG(LogJson, Warning, TEXT("JsonArrayStringToUStruct - Error parsing one of the elements. json=[%s]"), *JsonString);
+			return false;
+		}
+		return true;
+	}
+
+	/**
+	* Converts from an array of json values to an array of UStructs.
+	*
+	* @param JsonArray Array containing json values to convert.
+	* @param OutStructArray The UStruct array to copy in to
+	* @param CheckFlags Only convert properties that match at least one of these flags. If 0 check all properties.
+	* @param SkipFlags Skip properties that match any of these flags.
+	*
+	* @return False if any of the matching elements are not an object, or if one of the matching elements could not be converted to the specified UStruct type.
+	*/
+	template<typename OutStructType>
+	static bool JsonArrayToUStruct(const TArray<TSharedPtr<FJsonValue>>& JsonArray, TArray<OutStructType>* OutStructArray, int64 CheckFlags, int64 SkipFlags)
+	{
 		OutStructArray->SetNum(JsonArray.Num());
 		for (int32 i = 0; i < JsonArray.Num(); ++i)
 		{
 			const auto& Value = JsonArray[i];
 			if (Value->Type != EJson::Object)
 			{
-				UE_LOG(LogJson, Warning, TEXT("JsonArrayStringToUStruct - One array element was not an object. json=[%s]"), *JsonString);
+				UE_LOG(LogJson, Warning, TEXT("JsonArrayToUStruct - Array element [%i] was not an object."), i);
 				return false;
 			}
 			if (!FJsonObjectConverter::JsonObjectToUStruct(Value->AsObject().ToSharedRef(), OutStructType::StaticStruct(), &(*OutStructArray)[i], CheckFlags, SkipFlags))
 			{
-				UE_LOG(LogJson, Warning, TEXT("JsonArrayStringToUStruct - Unable to deserialize. json=[%s]"), *JsonString);
+				UE_LOG(LogJson, Warning, TEXT("JsonArrayToUStruct - Unable to convert element [%i]."), i);
 				return false;
 			}
 		}
diff --git a/Engine/Source/Runtime/Online/OnlineSubsystem/Private/OnlineSubsystemModule.cpp b/Engine/Source/Runtime/Online/OnlineSubsystem/Private/OnlineSubsystemModule.cpp
--- a/Engine/Source/Runtime/Online/OnlineSubsystem/Private/OnlineSubsystemModule.cpp
+++ b/Engine/Source/Runtime/Online/OnlineSubsystem/Private/OnlineSubsystemModule.cpp
@@ -58,8 +58,9 @@ void FOnlineSubsystemModule::StartupModule()
 {
 	FString InterfaceString;
 
-	// Load the platform defined "default" online services module
-	if (GConfig->GetString(TEXT("OnlineSubsystem"), TEXT("DefaultPlatformService"), InterfaceString, GEngineIni) &&
+	// Try to load the right service defined by an overload as a command line argument or as a value in the config
+	if ((FParse::Value(FCommandLine::Get(), *FString::Printf(TEXT("PlatformService=")), InterfaceString) ||
+		GConfig->GetString(TEXT("OnlineSubsystem"), TEXT("DefaultPlatformService"), InterfaceString, GEngineIni)) &&
 		InterfaceString.Len() > 0)
 	{
 		FName InterfaceName = FName(*InterfaceString);
