#json Dom builder improvements

diff --git a/Engine/Source/Runtime/JsonUtilities/Public/JsonDomBuilder.h b/Engine/Source/Runtime/JsonUtilities/Public/JsonDomBuilder.h
--- a/Engine/Source/Runtime/JsonUtilities/Public/JsonDomBuilder.h
+++ b/Engine/Source/Runtime/JsonUtilities/Public/JsonDomBuilder.h
@@ -4,6 +4,12 @@
 
 #include "Dom/JsonValue.h"
 #include "Dom/JsonObject.h"
+#include "Serialization/JsonSerializer.h"
+
+#include "Templates/IsFloatingPoint.h"
+#include "Templates/IsIntegral.h"
+#include "Templates/EnableIf.h"
+#include "Templates/Invoke.h"
 
 /**
  * Helpers for creating TSharedPtr<FJsonValue> JSON trees
@@ -40,20 +46,39 @@ public:
 			return MakeShared<FJsonValueObject>(Object);
 		}
 
-		FObject& Set(const FString& Key, const FArray& Arr)        { Object->SetField(Key, Arr.AsJsonValue());                      return *this; }
-		FObject& Set(const FString& Key, const FObject& Obj)       { Object->SetField(Key, Obj.AsJsonValue());                      return *this; }
+		template <template <class> class TPrintPolicy = TPrettyJsonPrintPolicy>
+		FString ToString() const
+		{
+			FString Result;
+			auto JsonWriter = TJsonWriterFactory<TCHAR, TPrintPolicy<TCHAR>>::Create(&Result);
+			FJsonSerializer::Serialize(Object, JsonWriter);
+			return Result;
+		}
+
+		FObject& Set(const FString& Key, const FArray& Arr)            { Object->SetField(Key, Arr.AsJsonValue());                      return *this; }
+		FObject& Set(const FString& Key, const FObject& Obj)           { Object->SetField(Key, Obj.AsJsonValue());                      return *this; }
 
-		FObject& Set(const FString& Key, const FString& Str)       { Object->SetField(Key, MakeShared<FJsonValueString>(Str));      return *this; }
+		FObject& Set(const FString& Key, const FString& Str)           { Object->SetField(Key, MakeShared<FJsonValueString>(Str));      return *this; }
 
 		template <class FNumber>
 		typename TEnableIf<TIsIntegral<FNumber>::Value || TIsFloatingPoint<FNumber>::Value, FObject&>::Type
-			Set(const FString& Key, FNumber Number)                { Object->SetField(Key, MakeShared<FJsonValueNumber>(Number));   return *this; }
+			Set(const FString& Key, FNumber Number)                    { Object->SetField(Key, MakeShared<FJsonValueNumber>(Number));   return *this; }
+
+		FObject& Set(const FString& Key, bool Boolean)                 { Object->SetField(Key, MakeShared<FJsonValueBoolean>(Boolean)); return *this; }
+		FObject& Set(const FString& Key, TYPE_OF_NULLPTR)              { Object->SetField(Key, MakeShared<FJsonValueNull>());           return *this; }
 
-		FObject& Set(const FString& Key, bool Boolean)             { Object->SetField(Key, MakeShared<FJsonValueBoolean>(Boolean)); return *this; }
-		FObject& Set(const FString& Key, TYPE_OF_NULLPTR)          { Object->SetField(Key, MakeShared<FJsonValueNull>());           return *this; }
+		FObject& Set(const FString& Key, TSharedPtr<FJsonValue> Value) { Object->SetField(Key, Value);                                  return *this; }
 
-		template <class FValue>
-		FObject& Set(const FString& Key, TSharedPtr<FValue> Value) { Object->SetField(Key, Value);                                  return *this; }
+		void CopyIf(const FJsonObject& Src, TFunctionRef<bool (const FString&, const FJsonValue&)> Pred)
+		{
+			for (const TPair<FString, TSharedPtr<FJsonValue>>& KV: Src.Values)
+			{
+				if (ensure(KV.Value) && Pred(KV.Key, *KV.Value))
+				{
+					Object->SetField(KV.Key, KV.Value);
+				}
+			}
+		}
 
 	private:
 		TSharedRef<FJsonObject> Object;
@@ -67,21 +92,30 @@ public:
 			return MakeShared<FJsonValueArray>(Array);
 		}
 
-		FArray& Add(const FArray& Arr)        { Array.Emplace(Arr.AsJsonValue());                      return *this; }
-		FArray& Add(const FObject& Obj)       { Array.Emplace(Obj.AsJsonValue());                      return *this; }
+		template <template <class> class TPrintPolicy = TPrettyJsonPrintPolicy>
+		FString ToString() const
+		{
+			FString Result;
+			auto JsonWriter = TJsonWriterFactory<TCHAR, TPrintPolicy<TCHAR>>::Create(&Result);
+			FJsonSerializer::Serialize(Array, JsonWriter);
+			return Result;
+		}
+
+		FArray& Add(const FArray& Arr)            { Array.Emplace(Arr.AsJsonValue());                      return *this; }
+		FArray& Add(const FObject& Obj)           { Array.Emplace(Obj.AsJsonValue());                      return *this; }
 
-		FArray& Add(const FString& Str)       { Array.Emplace(MakeShared<FJsonValueString>(Str));      return *this; }
+		FArray& Add(const FString& Str)           { Array.Emplace(MakeShared<FJsonValueString>(Str));      return *this; }
 
 		template <class FNumber>
 		typename TEnableIf<TIsIntegral<FNumber>::Value || TIsFloatingPoint<FNumber>::Value, FArray&>::Type
-			Add(FNumber Number)               { Array.Emplace(MakeShared<FJsonValueNumber>(Number));   return *this; }
+			Add(FNumber Number)                   { Array.Emplace(MakeShared<FJsonValueNumber>(Number));   return *this; }
 
-		FArray& Add(bool Boolean)             { Array.Emplace(MakeShared<FJsonValueBoolean>(Boolean)); return *this; }
-		FArray& Add(TYPE_OF_NULLPTR)          { Array.Emplace(MakeShared<FJsonValueNull>());           return *this; }
+		FArray& Add(bool Boolean)                 { Array.Emplace(MakeShared<FJsonValueBoolean>(Boolean)); return *this; }
+		FArray& Add(TYPE_OF_NULLPTR)              { Array.Emplace(MakeShared<FJsonValueNull>());           return *this; }
 
-		template <class FValue>
-		FArray& Add(TSharedPtr<FValue> Value) { Array.Emplace(Value);                                  return *this; }
+		FArray& Add(TSharedPtr<FJsonValue> Value) { Array.Emplace(Value);                                  return *this; }
 
+		/** Add multiple values */
 		template <class... FValue>
 		typename TEnableIf<(sizeof...(FValue) > 1), FArray&>::Type
 			Add(FValue&&... Value)
@@ -92,6 +126,16 @@ public:
 			return *this;
 		}
 
+		void CopyIf(const TArray<TSharedPtr<FJsonValue>>& Src, TFunctionRef<bool (const FJsonValue&)> Pred)
+		{
+			for (const TSharedPtr<FJsonValue>& Value: Src)
+			{
+				if (ensure(Value) && Pred(*Value))
+				{
+					Array.Emplace(Value);
+				}
+			}
+		}
 	private:
 		TArray<TSharedPtr<FJsonValue>> Array;
 	};
