#json Added Num accessors for dom builder

diff --git a/Engine/Source/Runtime/JsonUtilities/Public/JsonDomBuilder.h b/Engine/Source/Runtime/JsonUtilities/Public/JsonDomBuilder.h
--- a/Engine/Source/Runtime/JsonUtilities/Public/JsonDomBuilder.h
+++ b/Engine/Source/Runtime/JsonUtilities/Public/JsonDomBuilder.h
@@ -55,6 +55,11 @@ public:
 			return Result;
 		}
 
+		int Num() const
+		{
+			return Object->Values.Num();
+		}
+
 		FObject& Set(const FString& Key, const FArray& Arr)            { Object->SetField(Key, Arr.AsJsonValue());                      return *this; }
 		FObject& Set(const FString& Key, const FObject& Obj)           { Object->SetField(Key, Obj.AsJsonValue());                      return *this; }
 
@@ -101,6 +106,11 @@ public:
 			return Result;
 		}
 
+		int Num() const
+		{
+			return Array.Num();
+		}
+
 		FArray& Add(const FArray& Arr)            { Array.Emplace(Arr.AsJsonValue());                      return *this; }
 		FArray& Add(const FObject& Obj)           { Array.Emplace(Obj.AsJsonValue());                      return *this; }
 
