match any regional string for a generic language request after exhausting the pool for exact matches in polyglot FText JSON

diff --git a/Engine/Source/Runtime/JsonUtilities/Private/JsonObjectConverter.cpp b/Engine/Source/Runtime/JsonUtilities/Private/JsonObjectConverter.cpp
--- a/Engine/Source/Runtime/JsonUtilities/Private/JsonObjectConverter.cpp
+++ b/Engine/Source/Runtime/JsonUtilities/Private/JsonObjectConverter.cpp
@@ -331,6 +331,28 @@ bool FJsonObjectConverter::GetTextFromObject(const TSharedRef<FJsonObject>& Obj,
 		}
 	}
 
+	// try again but only search on the locale region (in the localized data). This is a common omission (i.e. en-US source text should be used if no en is defined)
+	for (const FString& LocaleToMatch : CultureList)
+	{
+		int32 SeparatorPos;
+		// only consider base language entries in culture chain (i.e. "en")
+		if (!LocaleToMatch.FindChar('-', SeparatorPos))
+		{
+			for (const auto& Pair : Obj->Values)
+			{
+				// only consider coupled entries now (base ones would have been matched on first path) (i.e. "en-US")
+				if (Pair.Key.FindChar('-', SeparatorPos))
+				{
+					if (Pair.Key.StartsWith(LocaleToMatch))
+					{
+						TextOut = FText::FromString(Pair.Value->AsString());
+						return true;
+					}
+				}
+			}
+		}
+	}
+
 	// no luck, is this possibly an unrelated json object?
 	return false;
 }
