| add support for formatting named args "date-" and "datetime-" when parsing text from JSON.

diff --git a/Engine/Source/Runtime/JsonUtilities/Private/JsonObjectConverter.cpp b/Engine/Source/Runtime/JsonUtilities/Private/JsonObjectConverter.cpp
--- a/Engine/Source/Runtime/JsonUtilities/Private/JsonObjectConverter.cpp
+++ b/Engine/Source/Runtime/JsonUtilities/Private/JsonObjectConverter.cpp
@@ -853,8 +853,27 @@ FFormatNamedArguments FJsonObjectConverter::ParseTextArgumentsFromJson(const TSh
 				NamedArgs.Emplace(It.Key, It.Value->AsNumber());
 				break;
 			case EJson::String:
-				// culture invariant string
-				NamedArgs.Emplace(It.Key, FText::FromString(It.Value->AsString()));
+				if (It.Key.StartsWith(TEXT("date-")))
+				{
+					FDateTime Dte;
+					if (FDateTime::ParseIso8601(*It.Value->AsString(), Dte))
+					{
+						NamedArgs.Emplace(It.Key, FText::AsDate(Dte));
+					}
+				}
+				else if (It.Key.StartsWith(TEXT("datetime-")))
+				{
+					FDateTime Dte;
+					if (FDateTime::ParseIso8601(*It.Value->AsString(), Dte))
+					{
+						NamedArgs.Emplace(It.Key, FText::AsDateTime(Dte));
+					}
+				}
+				else
+				{
+					// culture invariant string
+					NamedArgs.Emplace(It.Key, FText::FromString(It.Value->AsString()));
+				}
 				break;
 			case EJson::Object:
 			{
