JsonObjectConverter:
- Adding GetTextFromField
- Changed GetTextFromObject to use this internally

diff --git a/Engine/Source/Runtime/JsonUtilities/Private/JsonObjectConverter.cpp b/Engine/Source/Runtime/JsonUtilities/Private/JsonObjectConverter.cpp
--- a/Engine/Source/Runtime/JsonUtilities/Private/JsonObjectConverter.cpp
+++ b/Engine/Source/Runtime/JsonUtilities/Private/JsonObjectConverter.cpp
@@ -836,64 +836,81 @@ bool FJsonObjectConverter::JsonAttributesToUStruct(const TMap< FString, TSharedP
 	return JsonAttributesToUStructWithContainer(JsonAttributes, StructDefinition, OutStruct, StructDefinition, OutStruct, CheckFlags, SkipFlags);
 }
 
-FFormatNamedArguments FJsonObjectConverter::ParseTextArgumentsFromJson(const TSharedPtr<const FJsonObject>& JsonObject)
+//static 
+bool FJsonObjectConverter::GetTextFromField(const FString& FieldName, const TSharedPtr<FJsonValue>& FieldValue, FText& TextOut)
 {
-	FFormatNamedArguments NamedArgs;
-	if (JsonObject.IsValid())
+	if (FieldValue.IsValid())
 	{
-		for (const auto& It : JsonObject->Values)
+		switch (FieldValue->Type)
 		{
-			if (!It.Value.IsValid())
-				continue;
-
-			switch (It.Value->Type)
-			{
 			case EJson::Number:
+			{
 				// number
-				NamedArgs.Emplace(It.Key, It.Value->AsNumber());
-				break;
+				TextOut = FText::AsNumber(FieldValue->AsNumber());
+				return true;
+			}
 			case EJson::String:
-				if (It.Key.StartsWith(TEXT("date-")))
+			{
+				if (FieldName.StartsWith(TEXT("date-")))
 				{
 					FDateTime Dte;
-					if (FDateTime::ParseIso8601(*It.Value->AsString(), Dte))
+					if (FDateTime::ParseIso8601(*FieldValue->AsString(), Dte))
 					{
-						NamedArgs.Emplace(It.Key, FText::AsDate(Dte));
+						TextOut = FText::AsDate(Dte);
+						return true;
 					}
 				}
-				else if (It.Key.StartsWith(TEXT("datetime-")))
+				else if (FieldName.StartsWith(TEXT("datetime-")))
 				{
 					FDateTime Dte;
-					if (FDateTime::ParseIso8601(*It.Value->AsString(), Dte))
+					if (FDateTime::ParseIso8601(*FieldValue->AsString(), Dte))
 					{
-						NamedArgs.Emplace(It.Key, FText::AsDateTime(Dte));
+						TextOut = FText::AsDateTime(Dte);
+						return true;
 					}
 				}
 				else
 				{
 					// culture invariant string
-					NamedArgs.Emplace(It.Key, FText::FromString(It.Value->AsString()));
+					TextOut = FText::FromString(FieldValue->AsString());
+					return true;
 				}
 				break;
+			}
 			case EJson::Object:
 			{
 				// localized string
-				FText TextOut;
-				if (FJsonObjectConverter::GetTextFromObject(It.Value->AsObject().ToSharedRef(), TextOut))
-				{
-					NamedArgs.Emplace(It.Key, TextOut);
-				}
-				else
+				if (FJsonObjectConverter::GetTextFromObject(FieldValue->AsObject().ToSharedRef(), TextOut))
 				{
-					UE_LOG(LogJson, Error, TEXT("Unable to apply Json parameter %s (could not parse object)"), *It.Key);
+					return true;
 				}
+
+				UE_LOG(LogJson, Error, TEXT("Unable to apply Json parameter %s (could not parse object)"), *FieldName);
+				break;
 			}
-			break;
 			default:
-				UE_LOG(LogJson, Error, TEXT("Unable to apply Json parameter %s (bad type)"), *It.Key);
+			{
+				UE_LOG(LogJson, Error, TEXT("Unable to apply Json parameter %s (bad type)"), *FieldName);
 				break;
 			}
 		}
 	}
+	return false;
+}
+
+FFormatNamedArguments FJsonObjectConverter::ParseTextArgumentsFromJson(const TSharedPtr<const FJsonObject>& JsonObject)
+{
+	FFormatNamedArguments NamedArgs;
+	if (JsonObject.IsValid())
+	{
+		for (const auto& It : JsonObject->Values)
+		{
+			FText TextValue;
+			if (GetTextFromField(It.Key, It.Value, TextValue))
+			{
+				NamedArgs.Emplace(It.Key, TextValue);
+			}
+		}
+	}
 	return NamedArgs;
 }
\ No newline at end of file
diff --git a/Engine/Source/Runtime/JsonUtilities/Public/JsonObjectConverter.h b/Engine/Source/Runtime/JsonUtilities/Public/JsonObjectConverter.h
--- a/Engine/Source/Runtime/JsonUtilities/Public/JsonObjectConverter.h
+++ b/Engine/Source/Runtime/JsonUtilities/Public/JsonObjectConverter.h
@@ -21,6 +21,9 @@ public:
 	/** Parse an FText from a json object (assumed to be of the form where keys are culture codes and values are strings) */
 	static bool GetTextFromObject(const TSharedRef<FJsonObject>& Obj, FText& TextOut);
 
+	/** Convert a Json value to text (takes some hints from the value name) */
+	static bool GetTextFromField(const FString& FieldName, const TSharedPtr<FJsonValue>& FieldValue, FText& TextOut);
+
 public: // UStruct -> JSON
 
 	/**
