#json Added AsJsonObject accessor to FJsonDomBuilder::FObject
#rnx

diff --git a/Engine/Source/Runtime/JsonUtilities/Public/JsonDomBuilder.h b/Engine/Source/Runtime/JsonUtilities/Public/JsonDomBuilder.h
--- a/Engine/Source/Runtime/JsonUtilities/Public/JsonDomBuilder.h
+++ b/Engine/Source/Runtime/JsonUtilities/Public/JsonDomBuilder.h
@@ -46,6 +46,11 @@ public:
 			return MakeShared<FJsonValueObject>(Object);
 		}
 
+		TSharedRef<FJsonObject> AsJsonObject() const
+		{
+			return Object;
+		}
+
 		template <template <class> class TPrintPolicy = TPrettyJsonPrintPolicy>
 		FString ToString() const
 		{
