Use MakeShared instead of MakeShareable(new FJsonObject())

diff --git a/Engine/Source/Runtime/JsonUtilities/Public/JsonObjectConverter.h b/Engine/Source/Runtime/JsonUtilities/Public/JsonObjectConverter.h
--- a/Engine/Source/Runtime/JsonUtilities/Public/JsonObjectConverter.h
+++ b/Engine/Source/Runtime/JsonUtilities/Public/JsonObjectConverter.h
@@ -50,7 +50,7 @@ public: // UStruct -> JSON
 	template<typename InStructType>
 	static TSharedPtr<FJsonObject> UStructToJsonObject(const InStructType& InStruct, int64 CheckFlags = 0, int64 SkipFlags = 0, const CustomExportCallback* ExportCb = nullptr)
 	{
-		TSharedRef<FJsonObject> JsonObject = MakeShareable(new FJsonObject());
+		TSharedRef<FJsonObject> JsonObject = MakeShared<FJsonObject>();
 		if (UStructToJsonObject(InStructType::StaticStruct(), &InStruct, JsonObject, CheckFlags, SkipFlags, ExportCb))
 		{
 			return JsonObject;
