#json Make setting a null sharedptr work as would probably be expected, i.e. null JSON value

diff --git a/Engine/Source/Runtime/JsonUtilities/Public/JsonDomBuilder.h b/Engine/Source/Runtime/JsonUtilities/Public/JsonDomBuilder.h
--- a/Engine/Source/Runtime/JsonUtilities/Public/JsonDomBuilder.h
+++ b/Engine/Source/Runtime/JsonUtilities/Public/JsonDomBuilder.h
@@ -65,19 +65,19 @@ public:
 			return Object->Values.Num();
 		}
 
-		FObject& Set(const FString& Key, const FArray& Arr)            { Object->SetField(Key, Arr.AsJsonValue());                      return *this; }
-		FObject& Set(const FString& Key, const FObject& Obj)           { Object->SetField(Key, Obj.AsJsonValue());                      return *this; }
+		FObject& Set(const FString& Key, const FArray& Arr)            { Object->SetField(Key, Arr.AsJsonValue());                            return *this; }
+		FObject& Set(const FString& Key, const FObject& Obj)           { Object->SetField(Key, Obj.AsJsonValue());                            return *this; }
 
-		FObject& Set(const FString& Key, const FString& Str)           { Object->SetField(Key, MakeShared<FJsonValueString>(Str));      return *this; }
+		FObject& Set(const FString& Key, const FString& Str)           { Object->SetField(Key, MakeShared<FJsonValueString>(Str));            return *this; }
 
 		template <class FNumber>
 		typename TEnableIf<TIsIntegral<FNumber>::Value || TIsFloatingPoint<FNumber>::Value, FObject&>::Type
-			Set(const FString& Key, FNumber Number)                    { Object->SetField(Key, MakeShared<FJsonValueNumber>(Number));   return *this; }
+			Set(const FString& Key, FNumber Number)                    { Object->SetField(Key, MakeShared<FJsonValueNumber>(Number));         return *this; }
 
-		FObject& Set(const FString& Key, bool Boolean)                 { Object->SetField(Key, MakeShared<FJsonValueBoolean>(Boolean)); return *this; }
-		FObject& Set(const FString& Key, TYPE_OF_NULLPTR)              { Object->SetField(Key, MakeShared<FJsonValueNull>());           return *this; }
+		FObject& Set(const FString& Key, bool Boolean)                 { Object->SetField(Key, MakeShared<FJsonValueBoolean>(Boolean));       return *this; }
+		FObject& Set(const FString& Key, TYPE_OF_NULLPTR)              { Object->SetField(Key, MakeShared<FJsonValueNull>());                 return *this; }
 
-		FObject& Set(const FString& Key, TSharedPtr<FJsonValue> Value) { Object->SetField(Key, Value);                                  return *this; }
+		FObject& Set(const FString& Key, TSharedPtr<FJsonValue> Value) { Object->SetField(Key, Value ? Value : MakeShared<FJsonValueNull>()); return *this; }
 
 		void CopyIf(const FJsonObject& Src, TFunctionRef<bool (const FString&, const FJsonValue&)> Pred)
 		{
