#json Separated out serialisation functions

diff --git a/Engine/Source/Runtime/JsonUtilities/Private/JsonObjectWrapper.cpp b/Engine/Source/Runtime/JsonUtilities/Private/JsonObjectWrapper.cpp
--- a/Engine/Source/Runtime/JsonUtilities/Private/JsonObjectWrapper.cpp
+++ b/Engine/Source/Runtime/JsonUtilities/Private/JsonObjectWrapper.cpp
@@ -35,8 +35,7 @@ bool FJsonObjectWrapper::ImportTextItem(const TCHAR*& Buffer, int32 PortFlags, U
 	}
 
 	// parse the json
-	TSharedRef<TJsonReader<> > JsonReader = TJsonReaderFactory<>::Create(Json);
-	if (!FJsonSerializer::Deserialize(JsonReader, JsonObject) || !JsonObject.IsValid())
+	if (!JsonObjectFromString(Json))
 	{
 		if (ErrorText)
 		{
@@ -58,12 +57,7 @@ bool FJsonObjectWrapper::ExportTextItem(FString& ValueStr, FJsonObjectWrapper co
 	}
 
 	// serialize the json
-	TSharedRef<TJsonWriter<TCHAR, TCondensedJsonPrintPolicy<TCHAR> > > JsonWriter = TJsonWriterFactory<TCHAR, TCondensedJsonPrintPolicy<TCHAR> >::Create(&ValueStr, 0);
-	if (!FJsonSerializer::Serialize(JsonObject.ToSharedRef(), JsonWriter, true))
-	{
-		return false;
-	}
-	return true;
+	return JsonObjectToString(ValueStr);
 }
 
 void FJsonObjectWrapper::PostSerialize(const FArchive& Ar)
@@ -71,11 +65,23 @@ void FJsonObjectWrapper::PostSerialize(const FArchive& Ar)
 	if (!JsonString.IsEmpty())
 	{
 		// try to parse JsonString
-		TSharedRef<TJsonReader<> > JsonReader = TJsonReaderFactory<>::Create(JsonString);
-		if (!FJsonSerializer::Deserialize(JsonReader, JsonObject))
+		if (!JsonObjectFromString(JsonString))
 		{
 			// do not abide a string that won't parse
 			JsonString.Empty();
 		}
 	}
 }
+
+bool FJsonObjectWrapper::JsonObjectToString(FString& Str) const
+{
+	TSharedRef<TJsonWriter<TCHAR, TCondensedJsonPrintPolicy<TCHAR>>> JsonWriter = TJsonWriterFactory<TCHAR, TCondensedJsonPrintPolicy<TCHAR>>::Create(&Str, 0);
+	return FJsonSerializer::Serialize(JsonObject.ToSharedRef(), JsonWriter, true);
+}
+
+bool FJsonObjectWrapper::JsonObjectFromString(const FString& Str)
+{
+	TSharedRef<TJsonReader<>> JsonReader = TJsonReaderFactory<>::Create(Str);
+	return FJsonSerializer::Deserialize(JsonReader, JsonObject);
+}
+
diff --git a/Engine/Source/Runtime/JsonUtilities/Public/JsonObjectWrapper.h b/Engine/Source/Runtime/JsonUtilities/Public/JsonObjectWrapper.h
--- a/Engine/Source/Runtime/JsonUtilities/Public/JsonObjectWrapper.h
+++ b/Engine/Source/Runtime/JsonUtilities/Public/JsonObjectWrapper.h
@@ -31,6 +31,9 @@ public:
 	{
 		return JsonObject.IsValid();
 	}
+
+	bool JsonObjectToString(FString& Str) const;
+	bool JsonObjectFromString(const FString& Str);
 };
 
 template<>
