Backout CL 10127870. Probable cause for EdTest ensure

diff --git a/Engine/Source/Runtime/JsonUtilities/Private/JsonObjectConverter.cpp b/Engine/Source/Runtime/JsonUtilities/Private/JsonObjectConverter.cpp
--- a/Engine/Source/Runtime/JsonUtilities/Private/JsonObjectConverter.cpp
+++ b/Engine/Source/Runtime/JsonUtilities/Private/JsonObjectConverter.cpp
@@ -184,7 +184,7 @@ TSharedPtr<FJsonValue> FJsonObjectConverter::ObjectJsonCallback(UProperty* Prope
 {
 	if (UObjectProperty* ObjectProperty = Cast<UObjectProperty>(Property))
 	{
-		if (!ObjectProperty->HasAnyPropertyFlags(CPF_Transient)) // We are taking Transient to mean we don't want to serialize to Json either (could make a new flag if nessasary)
+		if (!ObjectProperty->HasAnyFlags(RF_Transient)) // We are taking Transient to mean we don't want to serialize to Json either (could make a new flag if nessasary)
 		{
 			TSharedRef<FJsonObject> Out = MakeShared<FJsonObject>();
 
