Deprecation of FJsonObjectConverter::ObjectJsonCallback(), which is unused and can cause problems because of a lack of null check, infinite recusion and the wrong transient flag being tested.

diff --git a/Engine/Source/Runtime/JsonUtilities/Private/JsonObjectConverter.cpp b/Engine/Source/Runtime/JsonUtilities/Private/JsonObjectConverter.cpp
--- a/Engine/Source/Runtime/JsonUtilities/Private/JsonObjectConverter.cpp
+++ b/Engine/Source/Runtime/JsonUtilities/Private/JsonObjectConverter.cpp
@@ -180,6 +180,8 @@ TSharedPtr<FJsonValue> ConvertScalarFPropertyToJsonValue(FProperty* Property, co
 }
 }
 
+PRAGMA_DISABLE_DEPRECATION_WARNINGS
+
 TSharedPtr<FJsonValue> FJsonObjectConverter::ObjectJsonCallback(FProperty* Property, const void* Value)
 {
 	if (FObjectProperty* ObjectProperty = CastField<FObjectProperty>(Property))
@@ -204,6 +206,8 @@ TSharedPtr<FJsonValue> FJsonObjectConverter::ObjectJsonCallback(FProperty* Prope
 	return TSharedPtr<FJsonValue>();
 }
 
+PRAGMA_ENABLE_DEPRECATION_WARNINGS
+
 TSharedPtr<FJsonValue> FJsonObjectConverter::UPropertyToJsonValue(FProperty* Property, const void* Value, int64 CheckFlags, int64 SkipFlags, const CustomExportCallback* ExportCb)
 {
 	if (Property->ArrayDim == 1)
diff --git a/Engine/Source/Runtime/JsonUtilities/Public/JsonObjectConverter.h b/Engine/Source/Runtime/JsonUtilities/Public/JsonObjectConverter.h
--- a/Engine/Source/Runtime/JsonUtilities/Public/JsonObjectConverter.h
+++ b/Engine/Source/Runtime/JsonUtilities/Public/JsonObjectConverter.h
@@ -36,6 +36,7 @@ public: // UStruct -> JSON
 	/**
 	 * Utility Export Callback for having object properties expanded to full Json.
 	 */
+	UE_DEPRECATED(4.25, "ObjectJsonCallback has been deprecated - please remove the usage of it from your project")
 	static TSharedPtr<FJsonValue> ObjectJsonCallback(FProperty* Property , const void* Value);
 
 	/**
