JsonObjectWrapper DisplayName set to JsonObject. Purely cosmetic.

diff --git a/Engine/Source/Runtime/JsonUtilities/Public/JsonObjectWrapper.h b/Engine/Source/Runtime/JsonUtilities/Public/JsonObjectWrapper.h
--- a/Engine/Source/Runtime/JsonUtilities/Public/JsonObjectWrapper.h
+++ b/Engine/Source/Runtime/JsonUtilities/Public/JsonObjectWrapper.h
@@ -12,7 +12,7 @@
 class FJsonObject;
 
 /** UStruct that holds a JsonObject, can be used by structs passed to JsonObjectConverter to pass through JsonObjects directly */
-USTRUCT(BlueprintType)
+USTRUCT(BlueprintType, meta = (DisplayName = "JsonObject"))
 struct JSONUTILITIES_API FJsonObjectWrapper
 {
 	GENERATED_USTRUCT_BODY()
