Fixed issue in Json to object serialization for Set properties, where duplicate entries could be created. The TSet container is now emptied prior to serializing in elements.

diff --git a/Engine/Source/Runtime/JsonUtilities/Private/JsonObjectConverter.cpp b/Engine/Source/Runtime/JsonUtilities/Private/JsonObjectConverter.cpp
--- a/Engine/Source/Runtime/JsonUtilities/Private/JsonObjectConverter.cpp
+++ b/Engine/Source/Runtime/JsonUtilities/Private/JsonObjectConverter.cpp
@@ -552,6 +552,7 @@ namespace
 			int32 ArrLen = ArrayValue.Num();
 
 			FScriptSetHelper Helper(SetProperty, OutValue);
+			Helper.EmptyElements(ArrLen);
 
 			// set the property values
 			for (int32 i = 0; i < ArrLen; ++i)
