JsonObjectWrapper initialization without loading from string

diff --git a/Engine/Source/Runtime/JsonUtilities/Private/JsonObjectWrapper.cpp b/Engine/Source/Runtime/JsonUtilities/Private/JsonObjectWrapper.cpp
--- a/Engine/Source/Runtime/JsonUtilities/Private/JsonObjectWrapper.cpp
+++ b/Engine/Source/Runtime/JsonUtilities/Private/JsonObjectWrapper.cpp
@@ -5,6 +5,11 @@
 #include "Serialization/JsonReader.h"
 #include "Serialization/JsonSerializer.h"
 
+FJsonObjectWrapper::FJsonObjectWrapper()
+{
+	JsonObject = MakeShared<FJsonObject>();
+}
+
 bool FJsonObjectWrapper::ImportTextItem(const TCHAR*& Buffer, int32 PortFlags, UObject* Parent, FOutputDevice* ErrorText)
 {
 	// read JSON string from Buffer
@@ -26,11 +31,11 @@ bool FJsonObjectWrapper::ImportTextItem(const TCHAR*& Buffer, int32 PortFlags, U
 		Buffer += Json.Len();
 	}
 
-	// empty string yields empty shared pointer
+	// empty string resets/re-initializes shared pointer
 	if (Json.IsEmpty())
 	{
 		JsonString.Empty();
-		JsonObject.Reset();
+		JsonObject = MakeShared<FJsonObject>();
 		return true;
 	}
 
@@ -73,6 +78,11 @@ void FJsonObjectWrapper::PostSerialize(const FArchive& Ar)
 	}
 }
 
+FJsonObjectWrapper::operator bool() const
+{
+	return JsonObject.IsValid() && !JsonObject->Values.IsEmpty();
+}
+
 bool FJsonObjectWrapper::JsonObjectToString(FString& Str) const
 {
 	TSharedRef<TJsonWriter<TCHAR, TCondensedJsonPrintPolicy<TCHAR>>> JsonWriter = TJsonWriterFactory<TCHAR, TCondensedJsonPrintPolicy<TCHAR>>::Create(&Str, 0);
diff --git a/Engine/Source/Runtime/JsonUtilities/Public/JsonObjectWrapper.h b/Engine/Source/Runtime/JsonUtilities/Public/JsonObjectWrapper.h
--- a/Engine/Source/Runtime/JsonUtilities/Public/JsonObjectWrapper.h
+++ b/Engine/Source/Runtime/JsonUtilities/Public/JsonObjectWrapper.h
@@ -16,8 +16,10 @@ USTRUCT(BlueprintType, meta = (DisplayName = "JsonObject"))
 struct JSONUTILITIES_API FJsonObjectWrapper
 {
 	GENERATED_USTRUCT_BODY()
+	
 public:
-
+	FJsonObjectWrapper();
+	
 	UPROPERTY(EditAnywhere, Category = "JSON")
 	FString JsonString;
 
@@ -27,10 +29,7 @@ public:
 	bool ExportTextItem(FString& ValueStr, FJsonObjectWrapper const& DefaultValue, UObject* Parent, int32 PortFlags, UObject* ExportRootScope) const;
 	void PostSerialize(const FArchive& Ar);
 
-	explicit operator bool() const
-	{
-		return JsonObject.IsValid();
-	}
+	explicit operator bool() const;
 
 	bool JsonObjectToString(FString& Str) const;
 	bool JsonObjectFromString(const FString& Str);
