Update Json serialisation to use long package names instead of short type names, to avoid warnings while loading custom properties

#rnx

diff --git a/Engine/Source/Runtime/JsonUtilities/Private/JsonObjectConverter.cpp b/Engine/Source/Runtime/JsonUtilities/Private/JsonObjectConverter.cpp
--- a/Engine/Source/Runtime/JsonUtilities/Private/JsonObjectConverter.cpp
+++ b/Engine/Source/Runtime/JsonUtilities/Private/JsonObjectConverter.cpp
@@ -182,7 +182,7 @@ TSharedPtr<FJsonValue> ConvertScalarFPropertyToJsonValue(FProperty* Property, co
 		{
 			TSharedRef<FJsonObject> Out = MakeShared<FJsonObject>();
 
-			Out->SetStringField(ObjectClassNameKey, Object->GetClass()->GetFName().ToString());
+			Out->SetStringField(ObjectClassNameKey, Object->GetClass()->GetPathName());
 			if (FJsonObjectConverter::UStructToJsonObject(ObjectProperty->GetObjectPropertyValue(Value)->GetClass(), Object, Out, CheckFlags, SkipFlags, ExportCb))
 			{
 				TSharedRef<FJsonValueObject> JsonObject = MakeShared<FJsonValueObject>(Out);
@@ -717,7 +717,7 @@ namespace
 				Obj->RemoveField(ObjectClassNameKey);
 				if (!ClassString.IsEmpty())
 				{
-					UClass* FoundClass = UClass::TryFindTypeSlow<UClass>(ClassString);
+					UClass* FoundClass = FPackageName::IsShortPackageName(ClassString) ? FindFirstObject<UClass>(*ClassString) : UClass::TryFindTypeSlow<UClass>(ClassString);
 					if (FoundClass)
 					{
 						PropertyClass = FoundClass;
