Fix some typos in code documentation

#rnx
PR #9399: Update JsonObjectConverter.cpp (Contributed by [REDACTED_PERSON])

diff --git a/Engine/Source/Runtime/JsonUtilities/Private/JsonObjectConverter.cpp b/Engine/Source/Runtime/JsonUtilities/Private/JsonObjectConverter.cpp
--- a/Engine/Source/Runtime/JsonUtilities/Private/JsonObjectConverter.cpp
+++ b/Engine/Source/Runtime/JsonUtilities/Private/JsonObjectConverter.cpp
@@ -14,7 +14,7 @@
 
 FString FJsonObjectConverter::StandardizeCase(const FString &StringIn)
 {
-	// this probably won't work for all cases, consider downcaseing the string fully
+	// this probably won't work for all cases, consider downcasing the string fully
 	FString FixedString = StringIn;
 	FixedString[0] = FChar::ToLower(FixedString[0]); // our json classes/variable start lower case
 	FixedString.ReplaceInline(TEXT("ID"), TEXT("Id"), ESearchCase::CaseSensitive); // Id is standard instead of ID, some of our fnames use ID
@@ -216,7 +216,7 @@ TSharedPtr<FJsonValue> FJsonObjectConverter::ObjectJsonCallback(FProperty* Prope
 {
 	if (FObjectProperty* ObjectProperty = CastField<FObjectProperty>(Property))
 	{
-		if (!ObjectProperty->HasAnyFlags(RF_Transient)) // We are taking Transient to mean we don't want to serialize to Json either (could make a new flag if nessasary)
+		if (!ObjectProperty->HasAnyFlags(RF_Transient)) // We are taking Transient to mean we don't want to serialize to Json either (could make a new flag if necessary)
 		{
 			TSharedRef<FJsonObject> Out = MakeShared<FJsonObject>();
 
@@ -989,4 +989,4 @@ FFormatNamedArguments FJsonObjectConverter::ParseTextArgumentsFromJson(const TSh
 		}
 	}
 	return NamedArgs;
-}
\ No newline at end of file
+}
