Better surfacing of json to uobject errors

diff --git a/Engine/Source/Runtime/JsonUtilities/Private/JsonObjectConverter.cpp b/Engine/Source/Runtime/JsonUtilities/Private/JsonObjectConverter.cpp
--- a/Engine/Source/Runtime/JsonUtilities/Private/JsonObjectConverter.cpp
+++ b/Engine/Source/Runtime/JsonUtilities/Private/JsonObjectConverter.cpp
@@ -12,11 +12,13 @@
 #include "Policies/CondensedJsonPrintPolicy.h"
 #include "JsonObjectWrapper.h"
 
+#define LOCTEXT_NAMESPACE "JsonObjectConverter"
+
 FString FJsonObjectConverter::StandardizeCase(const FString &StringIn)
 {
 	// this probably won't work for all cases, consider downcasing the string fully
 	FString FixedString = StringIn;
-	FixedString[0] = FChar::ToLower(FixedString[0]); // our json classes/variable start lower case
+	FixedString[0] = FChar::ToLower(FixedString[0]); // our JSON classes/variable start lower case
 	FixedString.ReplaceInline(TEXT("ID"), TEXT("Id"), ESearchCase::CaseSensitive); // Id is standard instead of ID, some of our fnames use ID
 	return FixedString;
 }
@@ -216,7 +218,7 @@ TSharedPtr<FJsonValue> FJsonObjectConverter::ObjectJsonCallback(FProperty* Prope
 {
 	if (FObjectProperty* ObjectProperty = CastField<FObjectProperty>(Property))
 	{
-		if (!ObjectProperty->HasAnyFlags(RF_Transient)) // We are taking Transient to mean we don't want to serialize to Json either (could make a new flag if necessary)
+		if (!ObjectProperty->HasAnyFlags(RF_Transient)) // We are taking Transient to mean we don't want to serialize to JSON either (could make a new flag if necessary)
 		{
 			TSharedRef<FJsonObject> Out = MakeShared<FJsonObject>();
 
@@ -340,7 +342,7 @@ bool FJsonObjectConverter::UStructToJsonObjectString(const UStruct* StructDefini
 		}
 		else
 		{
-			UE_LOG(LogJson, Warning, TEXT("UStructToJsonObjectString - Unable to write out json"));
+			UE_LOG(LogJson, Warning, TEXT("UStructToJsonObjectString - Unable to write out JSON"));
 		}
 	}
 
@@ -387,18 +389,18 @@ bool FJsonObjectConverter::GetTextFromObject(const TSharedRef<FJsonObject>& Obj,
 		}
 	}
 
-	// no luck, is this possibly an unrelated json object?
+	// no luck, is this possibly an unrelated JSON object?
 	return false;
 }
 
 
 namespace
 {
-	bool JsonValueToFPropertyWithContainer(const TSharedPtr<FJsonValue>& JsonValue, FProperty* Property, void* OutValue, const UStruct* ContainerStruct, void* Container, int64 CheckFlags, int64 SkipFlags, const bool bStrictMode);
-	bool JsonAttributesToUStructWithContainer(const TMap< FString, TSharedPtr<FJsonValue> >& JsonAttributes, const UStruct* StructDefinition, void* OutStruct, const UStruct* ContainerStruct, void* Container, int64 CheckFlags, int64 SkipFlags, const bool bStrictMode);
+	bool JsonValueToFPropertyWithContainer(const TSharedPtr<FJsonValue>& JsonValue, FProperty* Property, void* OutValue, const UStruct* ContainerStruct, void* Container, int64 CheckFlags, int64 SkipFlags, const bool bStrictMode, FText* OutFailReason);
+	bool JsonAttributesToUStructWithContainer(const TMap< FString, TSharedPtr<FJsonValue> >& JsonAttributes, const UStruct* StructDefinition, void* OutStruct, const UStruct* ContainerStruct, void* Container, int64 CheckFlags, int64 SkipFlags, const bool bStrictMode, FText* OutFailReason);
 
 	/** Convert JSON to property, assuming either the property is not an array or the value is an individual array element */
-	bool ConvertScalarJsonValueToFPropertyWithContainer(const TSharedPtr<FJsonValue>& JsonValue, FProperty* Property, void* OutValue, const UStruct* ContainerStruct, void* Container, int64 CheckFlags, int64 SkipFlags, const bool bStrictMode)
+	bool ConvertScalarJsonValueToFPropertyWithContainer(const TSharedPtr<FJsonValue>& JsonValue, FProperty* Property, void* OutValue, const UStruct* ContainerStruct, void* Container, int64 CheckFlags, int64 SkipFlags, const bool bStrictMode, FText* OutFailReason)
 	{
 	if (FEnumProperty* EnumProperty = CastField<FEnumProperty>(Property))
 	{
@@ -411,7 +413,11 @@ namespace
 			int64 IntValue = Enum->GetValueByName(FName(*StrValue), EGetByNameFlags::CheckAuthoredName);
 			if (IntValue == INDEX_NONE)
 			{
-				UE_LOG(LogJson, Error, TEXT("JsonValueToUProperty - Unable import enum %s from string value %s for property %s"), *Enum->CppType, *StrValue, *Property->GetNameCPP());
+				UE_LOG(LogJson, Error, TEXT("JsonValueToUProperty - Unable to import enum %s from string value %s for property %s"), *Enum->CppType, *StrValue, *Property->GetAuthoredName());
+				if (OutFailReason)
+				{
+					*OutFailReason = FText::Format(LOCTEXT("FailImportEnumFromString", "Unable to import enum {0} from string value {1} for property {2}"), FText::FromString(Enum->CppType), FText::FromString(StrValue), FText::FromString(Property->GetAuthoredName()));
+				}
 				return false;
 			}
 			EnumProperty->GetUnderlyingProperty()->SetIntPropertyValue(OutValue, IntValue);
@@ -433,7 +439,11 @@ namespace
 			int64 IntValue = Enum->GetValueByName(FName(*StrValue), EGetByNameFlags::CheckAuthoredName);
 			if (IntValue == INDEX_NONE)
 			{
-				UE_LOG(LogJson, Error, TEXT("JsonValueToUProperty - Unable import enum %s from string value %s for property %s"), *Enum->CppType, *StrValue, *Property->GetNameCPP());
+				UE_LOG(LogJson, Error, TEXT("JsonValueToUProperty - Unable to import enum %s from numeric value %s for property %s"), *Enum->CppType, *StrValue, *Property->GetAuthoredName());
+				if (OutFailReason)
+				{
+					*OutFailReason = FText::Format(LOCTEXT("FailImportEnumFromNumeric", "Unable to import enum {0} from numeric value {1} for property {2}"), FText::FromString(Enum->CppType), FText::FromString(StrValue), FText::FromString(Property->GetAuthoredName()));
+				}
 				return false;
 			}
 			NumericProperty->SetIntPropertyValue(OutValue, IntValue);
@@ -458,7 +468,11 @@ namespace
 		}
 		else
 		{
-			UE_LOG(LogJson, Error, TEXT("JsonValueToUProperty - Unable to set numeric property type %s for property %s"), *Property->GetClass()->GetName(), *Property->GetNameCPP());
+			UE_LOG(LogJson, Error, TEXT("JsonValueToUProperty - Unable to import json value into %s numeric property %s"), *Property->GetClass()->GetName(), *Property->GetAuthoredName());
+			if (OutFailReason)
+			{
+				*OutFailReason = FText::Format(LOCTEXT("FailImportNumericProperty", "Unable to import json value into {0} numeric property {1}"), FText::FromString(Property->GetClass()->GetName()), FText::FromString(Property->GetAuthoredName()));
+			}
 			return false;
 		}
 	}
@@ -489,9 +503,13 @@ namespace
 				const TSharedPtr<FJsonValue>& ArrayValueItem = ArrayValue[i];
 				if (ArrayValueItem.IsValid() && !ArrayValueItem->IsNull())
 				{
-					if (!JsonValueToFPropertyWithContainer(ArrayValueItem, ArrayProperty->Inner, Helper.GetRawPtr(i), ContainerStruct, Container, CheckFlags & (~CPF_ParmFlags), SkipFlags, bStrictMode))
+					if (!JsonValueToFPropertyWithContainer(ArrayValueItem, ArrayProperty->Inner, Helper.GetRawPtr(i), ContainerStruct, Container, CheckFlags & (~CPF_ParmFlags), SkipFlags, bStrictMode, OutFailReason))
 					{
-						UE_LOG(LogJson, Error, TEXT("JsonValueToUProperty - Unable to deserialize array element [%d] for property %s"), i, *Property->GetNameCPP());
+						UE_LOG(LogJson, Error, TEXT("JsonValueToUProperty - Unable to import Array element %d for property %s"), i, *Property->GetAuthoredName());
+						if (OutFailReason)
+						{
+							*OutFailReason = FText::Format(LOCTEXT("FailImportArrayElement", "Unable to import Array element {0} for property {1}\n{2}"), FText::AsNumber(i), FText::FromString(Property->GetAuthoredName()), *OutFailReason);
+						}
 						return false;
 					}
 				}
@@ -499,7 +517,11 @@ namespace
 		}
 		else
 		{
-			UE_LOG(LogJson, Error, TEXT("JsonValueToUProperty - Attempted to import TArray from non-array JSON key for property %s"), *Property->GetNameCPP());
+			UE_LOG(LogJson, Error, TEXT("JsonValueToUProperty - Unable to import non-array JSON value into Array property %s"), *Property->GetAuthoredName());
+			if (OutFailReason)
+			{
+				*OutFailReason = FText::Format(LOCTEXT("FailImportArray", "Unable to import non-array JSON value into Array property {0}"), FText::FromString(Property->GetAuthoredName()));
+			}
 			return false;
 		}
 	}
@@ -525,12 +547,23 @@ namespace
 
 					TSharedPtr<FJsonValueString> TempKeyValue = MakeShared<FJsonValueString>(Entry.Key);
 
-					const bool bKeySuccess = JsonValueToFPropertyWithContainer(TempKeyValue, MapProperty->KeyProp, Helper.GetKeyPtr(NewIndex), ContainerStruct, Container, CheckFlags & (~CPF_ParmFlags), SkipFlags, bStrictMode);
-					const bool bValueSuccess = JsonValueToFPropertyWithContainer(Entry.Value, MapProperty->ValueProp, Helper.GetValuePtr(NewIndex), ContainerStruct, Container, CheckFlags & (~CPF_ParmFlags), SkipFlags, bStrictMode);
+					if (!JsonValueToFPropertyWithContainer(TempKeyValue, MapProperty->KeyProp, Helper.GetKeyPtr(NewIndex), ContainerStruct, Container, CheckFlags & (~CPF_ParmFlags), SkipFlags, bStrictMode, OutFailReason))
+					{
+						UE_LOG(LogJson, Error, TEXT("JsonValueToUProperty - Unable to import Map element %s key for property %s"), *Entry.Key, *Property->GetAuthoredName());
+						if (OutFailReason)
+						{
+							*OutFailReason = FText::Format(LOCTEXT("FailImportMapElementKey", "Unable to import Map element {0} key for property {1}\n{2}"), FText::FromString(Entry.Key), FText::FromString(Property->GetAuthoredName()), *OutFailReason);
+						}
+						return false;
+					}
 
-					if (!(bKeySuccess && bValueSuccess))
+					if (!JsonValueToFPropertyWithContainer(Entry.Value, MapProperty->ValueProp, Helper.GetValuePtr(NewIndex), ContainerStruct, Container, CheckFlags & (~CPF_ParmFlags), SkipFlags, bStrictMode, OutFailReason))
 					{
-						UE_LOG(LogJson, Error, TEXT("JsonValueToUProperty - Unable to deserialize map element [key: %s] for property %s"), *Entry.Key, *Property->GetNameCPP());
+						UE_LOG(LogJson, Error, TEXT("JsonValueToUProperty - Unable to import Map element %s value for property %s"), *Entry.Key, *Property->GetAuthoredName());
+						if (OutFailReason)
+						{
+							*OutFailReason = FText::Format(LOCTEXT("FailImportMapElementValue", "Unable to import Map element {0} value for property {1}\n{2}"), FText::FromString(Entry.Key), FText::FromString(Property->GetAuthoredName()), *OutFailReason);
+						}
 						return false;
 					}
 				}
@@ -540,7 +573,11 @@ namespace
 		}
 		else
 		{
-			UE_LOG(LogJson, Error, TEXT("JsonValueToUProperty - Attempted to import TMap from non-object JSON key for property %s"), *Property->GetNameCPP());
+			UE_LOG(LogJson, Error, TEXT("JsonValueToUProperty - Unable to import non-object JSON value into Map property %s"), *Property->GetAuthoredName());
+			if (OutFailReason)
+			{
+				*OutFailReason = FText::Format(LOCTEXT("FailImportMap", "Unable to import non-object JSON value into Map property {0}"), FText::FromString(Property->GetAuthoredName()));
+			}
 			return false;
 		}
 	}
@@ -561,9 +598,13 @@ namespace
 				if (ArrayValueItem.IsValid() && !ArrayValueItem->IsNull())
 				{
 					int32 NewIndex = Helper.AddDefaultValue_Invalid_NeedsRehash();
-					if (!JsonValueToFPropertyWithContainer(ArrayValueItem, SetProperty->ElementProp, Helper.GetElementPtr(NewIndex), ContainerStruct, Container, CheckFlags & (~CPF_ParmFlags), SkipFlags, bStrictMode))
+					if (!JsonValueToFPropertyWithContainer(ArrayValueItem, SetProperty->ElementProp, Helper.GetElementPtr(NewIndex), ContainerStruct, Container, CheckFlags & (~CPF_ParmFlags), SkipFlags, bStrictMode, OutFailReason))
 					{
-						UE_LOG(LogJson, Error, TEXT("JsonValueToUProperty - Unable to deserialize set element [%d] for property %s"), i, *Property->GetNameCPP());
+						UE_LOG(LogJson, Error, TEXT("JsonValueToUProperty - Unable to import Set element %d for property %s"), i, *Property->GetAuthoredName());
+						if (OutFailReason)
+						{
+							*OutFailReason = FText::Format(LOCTEXT("FailImportSetElement", "Unable to import Set element {0} for property {1}\n{2}"), FText::AsNumber(i), FText::FromString(Property->GetAuthoredName()), *OutFailReason);
+						}
 						return false;
 					}
 				}
@@ -573,7 +614,11 @@ namespace
 		}
 		else
 		{
-			UE_LOG(LogJson, Error, TEXT("JsonValueToUProperty - Attempted to import TSet from non-array JSON key for property %s"), *Property->GetNameCPP());
+			UE_LOG(LogJson, Error, TEXT("JsonValueToUProperty - Unable to import non-array JSON value into Set property %s"), *Property->GetAuthoredName());
+			if (OutFailReason)
+			{
+				*OutFailReason = FText::Format(LOCTEXT("FailImportSet", "Unable to import non-array JSON value into Set property {0}"), FText::FromString(Property->GetAuthoredName()));
+			}
 			return false;
 		}
 	}
@@ -593,14 +638,22 @@ namespace
 			FText Text;
 			if (!FJsonObjectConverter::GetTextFromObject(Obj.ToSharedRef(), Text))
 			{
-				UE_LOG(LogJson, Error, TEXT("JsonValueToUProperty - Attempted to import FText from JSON object with invalid keys for property %s"), *Property->GetNameCPP());
+				UE_LOG(LogJson, Error, TEXT("JsonValueToUProperty - Unable to import JSON object with invalid keys into Text property %s"), *Property->GetAuthoredName());
+				if (OutFailReason)
+				{
+					*OutFailReason = FText::Format(LOCTEXT("FailImportTextFromObject", "Unable to import JSON object with invalid keys into Text property {0}"), FText::FromString(Property->GetAuthoredName()));
+				}
 				return false;
 			}
 			TextProperty->SetPropertyValue(OutValue, Text);
 		}
 		else
 		{
-			UE_LOG(LogJson, Error, TEXT("JsonValueToUProperty - Attempted to import FText from JSON that was neither string nor object for property %s"), *Property->GetNameCPP());
+			UE_LOG(LogJson, Error, TEXT("JsonValueToUProperty - Unable to import JSON value that is neither string nor object into Text property %s"), *Property->GetAuthoredName());
+			if (OutFailReason)
+			{
+				*OutFailReason = FText::Format(LOCTEXT("FailImportText", "Unable to import JSON value that is neither string nor object into Text property {0}"), FText::FromString(Property->GetAuthoredName()));
+			}
 			return false;
 		}
 	}
@@ -611,9 +664,13 @@ namespace
 		{
 			TSharedPtr<FJsonObject> Obj = JsonValue->AsObject();
 			check(Obj.IsValid()); // should not fail if Type == EJson::Object
-			if (!JsonAttributesToUStructWithContainer(Obj->Values, StructProperty->Struct, OutValue, ContainerStruct, Container, CheckFlags & (~CPF_ParmFlags), SkipFlags, bStrictMode))
+			if (!JsonAttributesToUStructWithContainer(Obj->Values, StructProperty->Struct, OutValue, ContainerStruct, Container, CheckFlags & (~CPF_ParmFlags), SkipFlags, bStrictMode, OutFailReason))
 			{
-				UE_LOG(LogJson, Error, TEXT("JsonValueToUProperty - FJsonObjectConverter::JsonObjectToUStruct failed for property %s"), *Property->GetNameCPP());
+				UE_LOG(LogJson, Error, TEXT("JsonValueToUProperty - Unable to import JSON object into %s property %s"), *StructProperty->Struct->GetAuthoredName(), *Property->GetAuthoredName());
+				if (OutFailReason)
+				{
+					*OutFailReason = FText::Format(LOCTEXT("FailImportStructFromObject", "Unable to import JSON object into {0} property {1}\n{2}"), FText::FromString(StructProperty->Struct->GetAuthoredName()), FText::FromString(Property->GetAuthoredName()), *OutFailReason);
+				}
 				return false;
 			}
 		}
@@ -650,7 +707,7 @@ namespace
 			}
 			else if (DateString == TEXT("now"))
 			{
-				// this value's not really meaningful from json serialization (since we don't know timezone) but handle it anyway since we're handling the other keywords
+				// this value's not really meaningful from JSON serialization (since we don't know timezone) but handle it anyway since we're handling the other keywords
 				DateTimeOut = FDateTime::UtcNow();
 			}
 			else if (FDateTime::ParseIso8601(*DateString, DateTimeOut))
@@ -663,7 +720,11 @@ namespace
 			}
 			else
 			{
-				UE_LOG(LogJson, Error, TEXT("JsonValueToUProperty - Unable to import FDateTime for property %s"), *Property->GetNameCPP());
+				UE_LOG(LogJson, Error, TEXT("JsonValueToUProperty - Unable to import JSON string into DateTime property %s"), *Property->GetAuthoredName());
+				if (OutFailReason)
+				{
+					*OutFailReason = FText::Format(LOCTEXT("FailImportDateTimeFromString", "Unable to import JSON string into DateTime property {0}"), FText::FromString(Property->GetAuthoredName()));
+				}
 				return false;
 			}
 		}
@@ -678,7 +739,11 @@ namespace
 				// Fall back to trying the tagged property approach if custom ImportTextItem couldn't get it done
 				if (Property->ImportText_Direct(ImportTextPtr, OutValue, nullptr, PPF_None) == nullptr)
 				{
-					UE_LOG(LogJson, Error, TEXT("JsonValueToUProperty - FJsonObjectConverter::JsonObjectToUStruct failed for property %s"), *Property->GetNameCPP());
+					UE_LOG(LogJson, Error, TEXT("JsonValueToUProperty - Unable to import JSON string into %s property %s"), *StructProperty->Struct->GetAuthoredName(), *Property->GetAuthoredName());
+					if (OutFailReason)
+					{
+						*OutFailReason = FText::Format(LOCTEXT("FailImportStructFromString", "Unable to import JSON string into {0} property {1}"), FText::FromString(StructProperty->Struct->GetAuthoredName()), FText::FromString(Property->GetAuthoredName()));
+					}
 					return false;
 				}
 			}
@@ -689,13 +754,21 @@ namespace
 			const TCHAR* ImportTextPtr = *ImportTextString;
 			if (Property->ImportText_Direct(ImportTextPtr, OutValue, nullptr, PPF_None) == nullptr)
 			{
-				UE_LOG(LogJson, Error, TEXT("JsonValueToUProperty - FJsonObjectConverter::JsonObjectToUStruct failed for property %s"), *Property->GetNameCPP());
+				UE_LOG(LogJson, Error, TEXT("JsonValueToUProperty - Unable to import JSON string into %s property %s"), *StructProperty->Struct->GetAuthoredName(), *Property->GetAuthoredName());
+				if (OutFailReason)
+				{
+					*OutFailReason = FText::Format(LOCTEXT("FailImportStructFromString", "Unable to import JSON string into {0} property {1}"), FText::FromString(StructProperty->Struct->GetAuthoredName()), FText::FromString(Property->GetAuthoredName()));
+				}
 				return false;
 			}
 		}
 		else
 		{
-			UE_LOG(LogJson, Error, TEXT("JsonValueToUProperty - Attempted to import UStruct from non-object JSON key for property %s"), *Property->GetNameCPP());
+			UE_LOG(LogJson, Error, TEXT("JsonValueToUProperty - Unable to import JSON value that is neither string nor object into %s property %s"), *StructProperty->Struct->GetAuthoredName(), *Property->GetAuthoredName());
+			if (OutFailReason)
+			{
+				*OutFailReason = FText::Format(LOCTEXT("FailImportStruct", "Unable to import JSON value that is neither string nor object into {0} property {1}"), FText::FromString(StructProperty->Struct->GetAuthoredName()), FText::FromString(Property->GetAuthoredName()));
+			}
 			return false;
 		}
 	}
@@ -712,7 +785,7 @@ namespace
 				TSharedPtr<FJsonObject> Obj = JsonValue->AsObject();
 				UClass* PropertyClass = ObjectProperty->PropertyClass;
 
-				// If a specific subclass was stored in the Json, use that instead of the PropertyClass
+				// If a specific subclass was stored in the JSON, use that instead of the PropertyClass
 				FString ClassString = Obj->GetStringField(ObjectClassNameKey);
 				Obj->RemoveField(ObjectClassNameKey);
 				if (!ClassString.IsEmpty())
@@ -730,9 +803,13 @@ namespace
 				ObjectProperty->SetObjectPropertyValue(OutValue, createdObj);
 
 				check(Obj.IsValid()); // should not fail if Type == EJson::Object
-				if (!JsonAttributesToUStructWithContainer(Obj->Values, PropertyClass, createdObj, PropertyClass, createdObj, CheckFlags & (~CPF_ParmFlags), SkipFlags, bStrictMode))
+				if (!JsonAttributesToUStructWithContainer(Obj->Values, PropertyClass, createdObj, PropertyClass, createdObj, CheckFlags & (~CPF_ParmFlags), SkipFlags, bStrictMode, OutFailReason))
 				{
-					UE_LOG(LogJson, Error, TEXT("JsonValueToUProperty - FJsonObjectConverter::JsonObjectToUStruct failed for property %s"), *Property->GetNameCPP());
+					UE_LOG(LogJson, Error, TEXT("JsonValueToUProperty - Unable to import JSON object into %s property %s"), *PropertyClass->GetAuthoredName(), *Property->GetAuthoredName());
+					if (OutFailReason)
+					{
+						*OutFailReason = FText::Format(LOCTEXT("FailImportObjectFromObject", "Unable to import JSON object into {0} property {1}\n{2}"), FText::FromString(PropertyClass->GetAuthoredName()), FText::FromString(Property->GetAuthoredName()), *OutFailReason);
+					}
 					return false;
 				}
 			}
@@ -741,7 +818,11 @@ namespace
 				// Default to expect a string for everything else
 				if (Property->ImportText_Direct(*JsonValue->AsString(), OutValue, nullptr, 0) == nullptr)
 				{
-					UE_LOG(LogJson, Error, TEXT("JsonValueToUProperty - Unable import property type %s from string value for property %s"), *Property->GetClass()->GetName(), *Property->GetNameCPP());
+					UE_LOG(LogJson, Error, TEXT("JsonValueToUProperty - Unable to import JSON string into %s property %s"), *ObjectProperty->PropertyClass->GetAuthoredName(), *Property->GetAuthoredName());
+					if (OutFailReason)
+					{
+						*OutFailReason = FText::Format(LOCTEXT("FailImportObjectFromString", "Unable to import JSON string into {0} property {1}"), FText::FromString(*ObjectProperty->PropertyClass->GetAuthoredName()), FText::FromString(Property->GetAuthoredName()));
+					}
 					return false;
 				}
 			}
@@ -751,7 +832,11 @@ namespace
 		// Default to expect a string for everything else
 		if (Property->ImportText_Direct(*JsonValue->AsString(), OutValue, nullptr, 0) == nullptr)
 		{
-			UE_LOG(LogJson, Error, TEXT("JsonValueToUProperty - Unable import property type %s from string value for property %s"), *Property->GetClass()->GetName(), *Property->GetNameCPP());
+			UE_LOG(LogJson, Error, TEXT("JsonValueToUProperty - Unable to import JSON string into property %s"), *Property->GetAuthoredName());
+			if (OutFailReason)
+			{
+				*OutFailReason = FText::Format(LOCTEXT("FailImportFromString", "Unable to import JSON string into property {0}"), FText::FromString(Property->GetAuthoredName()));
+			}
 			return false;
 		}
 	}
@@ -760,11 +845,15 @@ namespace
 	}
 
 
-	bool JsonValueToFPropertyWithContainer(const TSharedPtr<FJsonValue>& JsonValue, FProperty* Property, void* OutValue, const UStruct* ContainerStruct, void* Container, int64 CheckFlags, int64 SkipFlags, const bool bStrictMode)
+	bool JsonValueToFPropertyWithContainer(const TSharedPtr<FJsonValue>& JsonValue, FProperty* Property, void* OutValue, const UStruct* ContainerStruct, void* Container, int64 CheckFlags, int64 SkipFlags, const bool bStrictMode, FText* OutFailReason)
 	{
 		if (!JsonValue.IsValid())
 		{
-			UE_LOG(LogJson, Error, TEXT("JsonValueToUProperty - Invalid value JSON key"));
+			UE_LOG(LogJson, Error, TEXT("JsonValueToUProperty - Invalid JSON value"));
+			if (OutFailReason)
+			{
+				*OutFailReason = LOCTEXT("InvalidJsonValue", "Invalid JSON value");
+			}
 			return false;
 		}
 
@@ -775,7 +864,11 @@ namespace
 		{
 			if (bArrayOrSetProperty)
 			{
-				UE_LOG(LogJson, Error, TEXT("JsonValueToUProperty - Attempted to import TArray from non-array JSON key"));
+				UE_LOG(LogJson, Error, TEXT("JsonValueToUProperty - Expecting JSON array"));
+				if (OutFailReason)
+				{
+					*OutFailReason = LOCTEXT("ExpectingJsonArray", "Expecting JSON array");
+				}
 				return false;
 			}
 
@@ -784,20 +877,24 @@ namespace
 				if (bStrictMode)
 				{
 					UE_LOG(LogJson, Error, TEXT("JsonValueToUProperty - Property %s is not an array but has %d elements"), *Property->GetAuthoredName(), Property->ArrayDim);
+					if (OutFailReason)
+					{
+						*OutFailReason = FText::Format(LOCTEXT("InvalidDimensionOfNonArrayProperty", "Property {0} is not an array but has {1} elements"), FText::FromString(Property->GetAuthoredName()), FText::AsNumber(Property->ArrayDim));
+					}
 					return false;
 				}
 				
 				UE_LOG(LogJson, Warning, TEXT("Ignoring excess properties when deserializing %s"), *Property->GetAuthoredName());
 			}
 
-			return ConvertScalarJsonValueToFPropertyWithContainer(JsonValue, Property, OutValue, ContainerStruct, Container, CheckFlags, SkipFlags, bStrictMode);
+			return ConvertScalarJsonValueToFPropertyWithContainer(JsonValue, Property, OutValue, ContainerStruct, Container, CheckFlags, SkipFlags, bStrictMode, OutFailReason);
 		}
 
 		// In practice, the ArrayDim == 1 check ought to be redundant, since nested arrays of FProperties are not supported
 		if (bArrayOrSetProperty && Property->ArrayDim == 1)
 		{
 			// Read into TArray
-			return ConvertScalarJsonValueToFPropertyWithContainer(JsonValue, Property, OutValue, ContainerStruct, Container, CheckFlags, SkipFlags, bStrictMode);
+			return ConvertScalarJsonValueToFPropertyWithContainer(JsonValue, Property, OutValue, ContainerStruct, Container, CheckFlags, SkipFlags, bStrictMode, OutFailReason);
 		}
 
 		// We're deserializing a JSON array
@@ -805,7 +902,11 @@ namespace
 
 		if (bStrictMode && (Property->ArrayDim != ArrayValue.Num()))
 		{
-			UE_LOG(LogJson, Error, TEXT("JsonValueToUProperty - Json array for property %s doesn't have the same number of elements (has %d elements, but needs %d)"), *Property->GetAuthoredName(), ArrayValue.Num(), Property->ArrayDim);
+			UE_LOG(LogJson, Error, TEXT("JsonValueToUProperty - JSON array size is incorrect (has %d elements, but needs %d)"), ArrayValue.Num(), Property->ArrayDim);
+			if (OutFailReason)
+			{
+				*OutFailReason = FText::Format(LOCTEXT("IncorrectArraySize", "JSON array size is incorrect (has {0} elements, but needs {1})"), FText::AsNumber(ArrayValue.Num()), FText::AsNumber(Property->ArrayDim));
+			}
 			return false;
 		}
 		
@@ -818,7 +919,7 @@ namespace
 		const int32 ItemsToRead = FMath::Clamp(ArrayValue.Num(), 0, Property->ArrayDim);
 		for (int Index = 0; Index != ItemsToRead; ++Index)
 		{
-			if (!ConvertScalarJsonValueToFPropertyWithContainer(ArrayValue[Index], Property, static_cast<char*>(OutValue) + Index * Property->ElementSize, ContainerStruct, Container, CheckFlags, SkipFlags, bStrictMode))
+			if (!ConvertScalarJsonValueToFPropertyWithContainer(ArrayValue[Index], Property, static_cast<char*>(OutValue) + Index * Property->ElementSize, ContainerStruct, Container, CheckFlags, SkipFlags, bStrictMode, OutFailReason))
 			{
 				return false;
 			}
@@ -826,7 +927,7 @@ namespace
 		return true;
 	}
 
-	bool JsonAttributesToUStructWithContainer(const TMap< FString, TSharedPtr<FJsonValue> >& JsonAttributes, const UStruct* StructDefinition, void* OutStruct, const UStruct* ContainerStruct, void* Container, int64 CheckFlags, int64 SkipFlags, const bool bStrictMode)
+	bool JsonAttributesToUStructWithContainer(const TMap< FString, TSharedPtr<FJsonValue> >& JsonAttributes, const UStruct* StructDefinition, void* OutStruct, const UStruct* ContainerStruct, void* Container, int64 CheckFlags, int64 SkipFlags, const bool bStrictMode, FText* OutFailReason)
 	{
 		if (StructDefinition == FJsonObjectWrapper::StaticStruct())
 		{
@@ -858,7 +959,7 @@ namespace
 				continue;
 			}
 
-			// find a json value matching this property name
+			// find a JSON value matching this property name
 			FString PropertyName = StructDefinition->GetAuthoredNameForField(Property);
 			const TSharedPtr<FJsonValue>* JsonValue = JsonAttributes.Find(PropertyName);
 			
@@ -866,7 +967,11 @@ namespace
 			{
 				if (bStrictMode)
 				{
-					UE_LOG(LogJson, Error, TEXT("JsonObjectToUStruct - Unable to parse %s.%s from JSON"), *StructDefinition->GetAuthoredName(), *PropertyName);
+					UE_LOG(LogJson, Error, TEXT("JsonObjectToUStruct - Missing JSON value named %s"), *PropertyName);
+					if (OutFailReason)
+					{
+						*OutFailReason = FText::Format(LOCTEXT("MissingJsonField", "Missing JSON value named {0}"), FText::FromString(PropertyName));
+					}
 					return false;
 				}
 				
@@ -877,9 +982,13 @@ namespace
 			if (JsonValue->IsValid() && !(*JsonValue)->IsNull())
 			{
 				void* Value = Property->ContainerPtrToValuePtr<uint8>(OutStruct);
-				if (!JsonValueToFPropertyWithContainer(*JsonValue, Property, Value, ContainerStruct, Container, CheckFlags, SkipFlags, bStrictMode))
+				if (!JsonValueToFPropertyWithContainer(*JsonValue, Property, Value, ContainerStruct, Container, CheckFlags, SkipFlags, bStrictMode, OutFailReason))
 				{
-					UE_LOG(LogJson, Error, TEXT("JsonObjectToUStruct - Unable to parse %s.%s from JSON"), *StructDefinition->GetAuthoredName(), *PropertyName);
+					UE_LOG(LogJson, Error, TEXT("JsonObjectToUStruct - Unable to import JSON value into property %s"), *PropertyName);
+					if (OutFailReason)
+					{
+						*OutFailReason = FText::Format(LOCTEXT("FailImportValueToProperty", "Unable to import JSON value into property {0}\n{1}"), FText::FromString(PropertyName), *OutFailReason);
+					}
 					return false;
 				}
 			}
@@ -897,19 +1006,19 @@ namespace
 	}
 }
 
-bool FJsonObjectConverter::JsonValueToUProperty(const TSharedPtr<FJsonValue>& JsonValue, FProperty* Property, void* OutValue, int64 CheckFlags, int64 SkipFlags, const bool bStrictMode)
+bool FJsonObjectConverter::JsonValueToUProperty(const TSharedPtr<FJsonValue>& JsonValue, FProperty* Property, void* OutValue, int64 CheckFlags, int64 SkipFlags, const bool bStrictMode, FText* OutFailReason)
 {
-	return JsonValueToFPropertyWithContainer(JsonValue, Property, OutValue, nullptr, nullptr, CheckFlags, SkipFlags, bStrictMode);
+	return JsonValueToFPropertyWithContainer(JsonValue, Property, OutValue, nullptr, nullptr, CheckFlags, SkipFlags, bStrictMode, OutFailReason);
 }
 
-bool FJsonObjectConverter::JsonObjectToUStruct(const TSharedRef<FJsonObject>& JsonObject, const UStruct* StructDefinition, void* OutStruct, int64 CheckFlags, int64 SkipFlags, const bool bStrictMode)
+bool FJsonObjectConverter::JsonObjectToUStruct(const TSharedRef<FJsonObject>& JsonObject, const UStruct* StructDefinition, void* OutStruct, int64 CheckFlags, int64 SkipFlags, const bool bStrictMode, FText* OutFailReason)
 {
-	return JsonAttributesToUStruct(JsonObject->Values, StructDefinition, OutStruct, CheckFlags, SkipFlags, bStrictMode);
+	return JsonAttributesToUStruct(JsonObject->Values, StructDefinition, OutStruct, CheckFlags, SkipFlags, bStrictMode, OutFailReason);
 }
 
-bool FJsonObjectConverter::JsonAttributesToUStruct(const TMap< FString, TSharedPtr<FJsonValue> >& JsonAttributes, const UStruct* StructDefinition, void* OutStruct, int64 CheckFlags, int64 SkipFlags, const bool bStrictMode)
+bool FJsonObjectConverter::JsonAttributesToUStruct(const TMap< FString, TSharedPtr<FJsonValue> >& JsonAttributes, const UStruct* StructDefinition, void* OutStruct, int64 CheckFlags, int64 SkipFlags, const bool bStrictMode, FText* OutFailReason)
 {
-	return JsonAttributesToUStructWithContainer(JsonAttributes, StructDefinition, OutStruct, StructDefinition, OutStruct, CheckFlags, SkipFlags, bStrictMode);
+	return JsonAttributesToUStructWithContainer(JsonAttributes, StructDefinition, OutStruct, StructDefinition, OutStruct, CheckFlags, SkipFlags, bStrictMode, OutFailReason);
 }
 
 //static 
@@ -961,12 +1070,12 @@ bool FJsonObjectConverter::GetTextFromField(const FString& FieldName, const TSha
 					return true;
 				}
 
-				UE_LOG(LogJson, Error, TEXT("Unable to apply Json parameter %s (could not parse object)"), *FieldName);
+				UE_LOG(LogJson, Error, TEXT("Unable to apply JSON parameter %s (could not parse object)"), *FieldName);
 				break;
 			}
 			default:
 			{
-				UE_LOG(LogJson, Error, TEXT("Unable to apply Json parameter %s (bad type)"), *FieldName);
+				UE_LOG(LogJson, Error, TEXT("Unable to apply JSON parameter %s (bad type)"), *FieldName);
 				break;
 			}
 		}
@@ -990,3 +1099,4 @@ FFormatNamedArguments FJsonObjectConverter::ParseTextArgumentsFromJson(const TSh
 	}
 	return NamedArgs;
 }
+#undef LOCTEXT_NAMESPACE
diff --git a/Engine/Source/Runtime/JsonUtilities/Public/JsonObjectConverter.h b/Engine/Source/Runtime/JsonUtilities/Public/JsonObjectConverter.h
--- a/Engine/Source/Runtime/JsonUtilities/Public/JsonObjectConverter.h
+++ b/Engine/Source/Runtime/JsonUtilities/Public/JsonObjectConverter.h
@@ -188,10 +188,11 @@ public: // JSON -> UStruct
 	 * @param CheckFlags Only convert properties that match at least one of these flags. If 0 check all properties.
 	 * @param SkipFlags Skip properties that match any of these flags
 	 * @param bStrictMode Whether to strictly check the json attributes
+	 * @param OutFailReason Reason of the failure if any
 	 *
 	 * @return False if any properties matched but failed to deserialize
 	 */
-	static bool JsonObjectToUStruct(const TSharedRef<FJsonObject>& JsonObject, const UStruct* StructDefinition, void* OutStruct, int64 CheckFlags = 0, int64 SkipFlags = 0, const bool bStrictMode = false);
+	static bool JsonObjectToUStruct(const TSharedRef<FJsonObject>& JsonObject, const UStruct* StructDefinition, void* OutStruct, int64 CheckFlags = 0, int64 SkipFlags = 0, const bool bStrictMode = false, FText* OutFailReason = nullptr);
 
 	/**
 	 * Templated version of JsonObjectToUStruct
@@ -201,13 +202,14 @@ public: // JSON -> UStruct
 	 * @param CheckFlags Only convert properties that match at least one of these flags. If 0 check all properties.
 	 * @param SkipFlags Skip properties that match any of these flags
 	 * @param bStrictMode Whether to strictly check the json attributes
+	 * @param OutFailReason Reason of the failure if any
 	 *
 	 * @return False if any properties matched but failed to deserialize
 	 */
 	template<typename OutStructType>
-	static bool JsonObjectToUStruct(const TSharedRef<FJsonObject>& JsonObject, OutStructType* OutStruct, int64 CheckFlags = 0, int64 SkipFlags = 0, const bool bStrictMode = false)
+	static bool JsonObjectToUStruct(const TSharedRef<FJsonObject>& JsonObject, OutStructType* OutStruct, int64 CheckFlags = 0, int64 SkipFlags = 0, const bool bStrictMode = false, FText* OutFailReason = nullptr)
 	{
-		return JsonObjectToUStruct(JsonObject, OutStructType::StaticStruct(), OutStruct, CheckFlags, SkipFlags, bStrictMode);
+		return JsonObjectToUStruct(JsonObject, OutStructType::StaticStruct(), OutStruct, CheckFlags, SkipFlags, bStrictMode, OutFailReason);
 	}
 
 	/**
@@ -219,10 +221,11 @@ public: // JSON -> UStruct
 	 * @param CheckFlags Only convert properties that match at least one of these flags. If 0 check all properties.
 	 * @param SkipFlags Skip properties that match any of these flags
 	 * @param bStrictMode Whether to strictly check the json attributes
+	 * @param OutFailReason Reason of the failure if any
 	 *
 	 * @return False if any properties matched but failed to deserialize
 	 */
-	static bool JsonAttributesToUStruct(const TMap< FString, TSharedPtr<FJsonValue> >& JsonAttributes, const UStruct* StructDefinition, void* OutStruct, int64 CheckFlags = 0, int64 SkipFlags = 0, const bool bStrictMode = false);
+	static bool JsonAttributesToUStruct(const TMap< FString, TSharedPtr<FJsonValue> >& JsonAttributes, const UStruct* StructDefinition, void* OutStruct, int64 CheckFlags = 0, int64 SkipFlags = 0, const bool bStrictMode = false, FText* OutFailReason = nullptr);
 
 	/**
 	 * Converts a single JsonValue to the corresponding FProperty (this may recurse if the property is a UStruct for instance).
@@ -233,10 +236,11 @@ public: // JSON -> UStruct
 	 * @param CheckFlags Only convert sub-properties that match at least one of these flags. If 0 check all properties.
 	 * @param SkipFlags Skip sub-properties that match any of these flags
 	 * @param bStrictMode Whether to strictly check the json attributes
+	 * @param OutFailReason Reason of the failure if any
 	 *
 	 * @return False if the property failed to serialize
 	 */
-	static bool JsonValueToUProperty(const TSharedPtr<FJsonValue>& JsonValue, FProperty* Property, void* OutValue, int64 CheckFlags = 0, int64 SkipFlags = 0, const bool bStrictMode = false);
+	static bool JsonValueToUProperty(const TSharedPtr<FJsonValue>& JsonValue, FProperty* Property, void* OutValue, int64 CheckFlags = 0, int64 SkipFlags = 0, const bool bStrictMode = false, FText* OutFailReason = nullptr);
 
 	/**
 	 * Converts from a json string containing an object to a UStruct
