[json] Added utility export callback to write dates as ISO strings (would break backward compatibility to make it the default)

diff --git a/Engine/Source/Runtime/JsonUtilities/Private/JsonObjectConverter.cpp b/Engine/Source/Runtime/JsonUtilities/Private/JsonObjectConverter.cpp
--- a/Engine/Source/Runtime/JsonUtilities/Private/JsonObjectConverter.cpp
+++ b/Engine/Source/Runtime/JsonUtilities/Private/JsonObjectConverter.cpp
@@ -29,6 +29,8 @@ namespace
 {
 	const FString ObjectClassNameKey = "_ClassName";
 
+	const FName NAME_DateTime(TEXT("DateTime"));
+
 /** Convert property to JSON, assuming either the property is not an array or the value is an individual array element */
 TSharedPtr<FJsonValue> ConvertScalarFPropertyToJsonValue(FProperty* Property, const void* Value, int64 CheckFlags, int64 SkipFlags, const FJsonObjectConverter::CustomExportCallback* ExportCb, FProperty* OuterProperty)
 {
@@ -660,7 +662,6 @@ namespace
 	}
 	else if (FStructProperty *StructProperty = CastField<FStructProperty>(Property))
 	{
-		static const FName NAME_DateTime(TEXT("DateTime"));
 		if (JsonValue->Type == EJson::Object)
 		{
 			TSharedPtr<FJsonObject> Obj = JsonValue->AsObject();
@@ -1100,4 +1101,20 @@ FFormatNamedArguments FJsonObjectConverter::ParseTextArgumentsFromJson(const TSh
 	}
 	return NamedArgs;
 }
+
+const FJsonObjectConverter::CustomExportCallback FJsonObjectConverter::ExportCallback_WriteISO8601Dates = 
+	FJsonObjectConverter::CustomExportCallback::CreateLambda(
+		[](FProperty* Prop, const void* Data) -> TSharedPtr<FJsonValue>
+		{
+			if (FStructProperty* StructProperty = CastField<FStructProperty>(Prop))
+			{
+				checkSlow(StructProperty->Struct);
+				if (StructProperty->Struct->GetFName() == NAME_DateTime)
+				{
+					return MakeShared<FJsonValueString>(static_cast<const FDateTime*>(Data)->ToIso8601());
+				}
+			}
+			return {};
+		});
+
 #undef LOCTEXT_NAMESPACE
diff --git a/Engine/Source/Runtime/JsonUtilities/Public/JsonObjectConverter.h b/Engine/Source/Runtime/JsonUtilities/Public/JsonObjectConverter.h
--- a/Engine/Source/Runtime/JsonUtilities/Public/JsonObjectConverter.h
+++ b/Engine/Source/Runtime/JsonUtilities/Public/JsonObjectConverter.h
@@ -49,6 +49,8 @@ public: // UStruct -> JSON
 	 */
 	DECLARE_DELEGATE_RetVal_TwoParams(TSharedPtr<FJsonValue>, CustomExportCallback, FProperty* /* Property */, const void* /* Value */);
 
+	static const CustomExportCallback ExportCallback_WriteISO8601Dates;
+
 	/**
 	 * Utility Export Callback for having object properties expanded to full Json.
 	 */
