- Add constructor for existing json object

diff --git a/Engine/Source/Runtime/JsonUtilities/Public/JsonDomBuilder.h b/Engine/Source/Runtime/JsonUtilities/Public/JsonDomBuilder.h
--- a/Engine/Source/Runtime/JsonUtilities/Public/JsonDomBuilder.h
+++ b/Engine/Source/Runtime/JsonUtilities/Public/JsonDomBuilder.h
@@ -41,6 +41,16 @@ public:
 		{
 		}
 
+		FObject(TSharedPtr<FJsonObject> InObject)
+			: Object(InObject.ToSharedRef())
+		{
+		}
+
+		FObject(TSharedRef<FJsonObject> InObject)
+			: Object(InObject)
+		{
+		}
+
 		TSharedRef<FJsonValueObject> AsJsonValue() const
 		{
 			return MakeShared<FJsonValueObject>(Object);
