Add support for writing JSONs with original property name casings

Avoids cases like:
- https://forums.unrealengine.com/t/fjsonobjectconverter-ustructtojsonobjectstring/347721/8
- https://forums.unrealengine.com/t/property-names-are-changed-by-ustructtojsonobjectstring/318179/2

diff --git a/Engine/Source/Runtime/JsonUtilities/Private/JsonObjectConverter.cpp b/Engine/Source/Runtime/JsonUtilities/Private/JsonObjectConverter.cpp
--- a/Engine/Source/Runtime/JsonUtilities/Private/JsonObjectConverter.cpp
+++ b/Engine/Source/Runtime/JsonUtilities/Private/JsonObjectConverter.cpp
@@ -32,7 +32,7 @@ namespace
 	const FName NAME_DateTime(TEXT("DateTime"));
 
 /** Convert property to JSON, assuming either the property is not an array or the value is an individual array element */
-TSharedPtr<FJsonValue> ConvertScalarFPropertyToJsonValue(FProperty* Property, const void* Value, int64 CheckFlags, int64 SkipFlags, const FJsonObjectConverter::CustomExportCallback* ExportCb, FProperty* OuterProperty)
+TSharedPtr<FJsonValue> ConvertScalarFPropertyToJsonValue(FProperty* Property, const void* Value, int64 CheckFlags, int64 SkipFlags, const FJsonObjectConverter::CustomExportCallback* ExportCb, FProperty* OuterProperty, EJsonObjectConversionFlags ConversionFlags)
 {
 	// See if there's a custom export callback first, so it can override default behavior
 	if (ExportCb && ExportCb->IsBound())
@@ -132,8 +132,8 @@ TSharedPtr<FJsonValue> ConvertScalarFPropertyToJsonValue(FProperty* Property, co
 		{
 			if ( Helper.IsValidIndex(i) )
 			{
-				TSharedPtr<FJsonValue> KeyElement = FJsonObjectConverter::UPropertyToJsonValue(MapProperty->KeyProp, Helper.GetKeyPtr(i), CheckFlags & ( ~CPF_ParmFlags ), SkipFlags, ExportCb, MapProperty);
-				TSharedPtr<FJsonValue> ValueElement = FJsonObjectConverter::UPropertyToJsonValue(MapProperty->ValueProp, Helper.GetValuePtr(i), CheckFlags & ( ~CPF_ParmFlags ), SkipFlags, ExportCb, MapProperty);
+				TSharedPtr<FJsonValue> KeyElement = FJsonObjectConverter::UPropertyToJsonValue(MapProperty->KeyProp, Helper.GetKeyPtr(i), CheckFlags & ( ~CPF_ParmFlags ), SkipFlags, ExportCb, MapProperty, ConversionFlags);
+				TSharedPtr<FJsonValue> ValueElement = FJsonObjectConverter::UPropertyToJsonValue(MapProperty->ValueProp, Helper.GetValuePtr(i), CheckFlags & ( ~CPF_ParmFlags ), SkipFlags, ExportCb, MapProperty, ConversionFlags);
 				if ( KeyElement.IsValid() && ValueElement.IsValid() )
 				{
 					FString KeyString;
@@ -151,7 +151,10 @@ TSharedPtr<FJsonValue> ConvertScalarFPropertyToJsonValue(FProperty* Property, co
 					if (CastField<FEnumProperty>(MapProperty->KeyProp) ||
 						CastField<FNameProperty>(MapProperty->KeyProp))
 					{
-						KeyString = FJsonObjectConverter::StandardizeCase(KeyString);
+						if (!EnumHasAnyFlags(ConversionFlags, EJsonObjectConversionFlags::SkipStandardizeCase))
+						{
+							KeyString = FJsonObjectConverter::StandardizeCase(KeyString);
+						}
 					}
 					Out->SetField(KeyString, ValueElement);
 				}
@@ -174,7 +177,7 @@ TSharedPtr<FJsonValue> ConvertScalarFPropertyToJsonValue(FProperty* Property, co
 		}
 
 		TSharedRef<FJsonObject> Out = MakeShared<FJsonObject>();
-		if (FJsonObjectConverter::UStructToJsonObject(StructProperty->Struct, Value, Out, CheckFlags & (~CPF_ParmFlags), SkipFlags, ExportCb))
+		if (FJsonObjectConverter::UStructToJsonObject(StructProperty->Struct, Value, Out, CheckFlags & (~CPF_ParmFlags), SkipFlags, ExportCb, ConversionFlags))
 		{
 			return MakeShared<FJsonValueObject>(Out);
 		}
@@ -243,27 +246,27 @@ TSharedPtr<FJsonValue> FJsonObjectConverter::ObjectJsonCallback(FProperty* Prope
 
 PRAGMA_ENABLE_DEPRECATION_WARNINGS
 
-TSharedPtr<FJsonValue> FJsonObjectConverter::UPropertyToJsonValue(FProperty* Property, const void* Value, int64 CheckFlags, int64 SkipFlags, const CustomExportCallback* ExportCb, FProperty* OuterProperty)
+TSharedPtr<FJsonValue> FJsonObjectConverter::UPropertyToJsonValue(FProperty* Property, const void* Value, int64 CheckFlags, int64 SkipFlags, const CustomExportCallback* ExportCb, FProperty* OuterProperty, EJsonObjectConversionFlags ConversionFlags)
 {
 	if (Property->ArrayDim == 1)
 	{
-		return ConvertScalarFPropertyToJsonValue(Property, Value, CheckFlags, SkipFlags, ExportCb, OuterProperty);
+		return ConvertScalarFPropertyToJsonValue(Property, Value, CheckFlags, SkipFlags, ExportCb, OuterProperty, ConversionFlags);
 	}
 
 	TArray< TSharedPtr<FJsonValue> > Array;
 	for (int Index = 0; Index != Property->ArrayDim; ++Index)
 	{
-		Array.Add(ConvertScalarFPropertyToJsonValue(Property, (char*)Value + Index * Property->ElementSize, CheckFlags, SkipFlags, ExportCb, OuterProperty));
+		Array.Add(ConvertScalarFPropertyToJsonValue(Property, (char*)Value + Index * Property->ElementSize, CheckFlags, SkipFlags, ExportCb, OuterProperty, ConversionFlags));
 	}
 	return MakeShared<FJsonValueArray>(Array);
 }
 
-bool FJsonObjectConverter::UStructToJsonObject(const UStruct* StructDefinition, const void* Struct, TSharedRef<FJsonObject> OutJsonObject, int64 CheckFlags, int64 SkipFlags, const CustomExportCallback* ExportCb)
+bool FJsonObjectConverter::UStructToJsonObject(const UStruct* StructDefinition, const void* Struct, TSharedRef<FJsonObject> OutJsonObject, int64 CheckFlags, int64 SkipFlags, const CustomExportCallback* ExportCb, EJsonObjectConversionFlags ConversionFlags)
 {
-	return UStructToJsonAttributes(StructDefinition, Struct, OutJsonObject->Values, CheckFlags, SkipFlags, ExportCb);
+	return UStructToJsonAttributes(StructDefinition, Struct, OutJsonObject->Values, CheckFlags, SkipFlags, ExportCb, ConversionFlags);
 }
 
-bool FJsonObjectConverter::UStructToJsonAttributes(const UStruct* StructDefinition, const void* Struct, TMap< FString, TSharedPtr<FJsonValue> >& OutJsonAttributes, int64 CheckFlags, int64 SkipFlags, const CustomExportCallback* ExportCb)
+bool FJsonObjectConverter::UStructToJsonAttributes(const UStruct* StructDefinition, const void* Struct, TMap< FString, TSharedPtr<FJsonValue> >& OutJsonAttributes, int64 CheckFlags, int64 SkipFlags, const CustomExportCallback* ExportCb, EJsonObjectConversionFlags ConversionFlags)
 {
 	if (SkipFlags == 0)
 	{
@@ -297,11 +300,16 @@ bool FJsonObjectConverter::UStructToJsonAttributes(const UStruct* StructDefiniti
 			continue;
 		}
 
-		FString VariableName = StandardizeCase(Property->GetAuthoredName());
+		FString VariableName = Property->GetAuthoredName();
+		if (!EnumHasAnyFlags(ConversionFlags, EJsonObjectConversionFlags::SkipStandardizeCase))
+		{
+			VariableName = StandardizeCase(VariableName);
+		}
+
 		const void* Value = Property->ContainerPtrToValuePtr<uint8>(Struct);
 
 		// convert the property to a FJsonValue
-		TSharedPtr<FJsonValue> JsonValue = UPropertyToJsonValue(Property, Value, CheckFlags, SkipFlags, ExportCb);
+		TSharedPtr<FJsonValue> JsonValue = UPropertyToJsonValue(Property, Value, CheckFlags, SkipFlags, ExportCb, nullptr, ConversionFlags);
 		if (!JsonValue.IsValid())
 		{
 			FFieldClass* PropClass = Property->GetClass();
diff --git a/Engine/Source/Runtime/JsonUtilities/Public/JsonObjectConverter.h b/Engine/Source/Runtime/JsonUtilities/Public/JsonObjectConverter.h
--- a/Engine/Source/Runtime/JsonUtilities/Public/JsonObjectConverter.h
+++ b/Engine/Source/Runtime/JsonUtilities/Public/JsonObjectConverter.h
@@ -23,6 +23,14 @@
 #include "Trace/Detail/Channel.h"
 #include "UObject/Class.h"
 
+enum class EJsonObjectConversionFlags
+{
+	None = 0,
+	SkipStandardizeCase = 1 << 0,
+};
+
+ENUM_CLASS_FLAGS(EJsonObjectConversionFlags)
+
 class FProperty;
 class UStruct;
 
@@ -82,14 +90,15 @@ public: // UStruct -> JSON
 	 *
 	 * @param StructDefinition UStruct definition that is looked over for properties
 	 * @param Struct The UStruct instance to copy out of
-	 * @param JsonObject Json Object to be filled in with data from the ustruct
+	 * @param OutJsonObject Json Object to be filled in with data from the ustruct
 	 * @param CheckFlags Only convert properties that match at least one of these flags. If 0 check all properties.
 	 * @param SkipFlags Skip properties that match any of these flags
 	 * @param ExportCb Optional callback to override export behavior, if this returns null it will fallback to the default
+	 * @param ConversionFlags Bitwise flags to customize the conversion behavior
 	 *
 	 * @return False if any properties failed to write
 	 */
-	static JSONUTILITIES_API bool UStructToJsonObject(const UStruct* StructDefinition, const void* Struct, TSharedRef<FJsonObject> OutJsonObject, int64 CheckFlags = 0, int64 SkipFlags = 0, const CustomExportCallback* ExportCb = nullptr);
+	static JSONUTILITIES_API bool UStructToJsonObject(const UStruct* StructDefinition, const void* Struct, TSharedRef<FJsonObject> OutJsonObject, int64 CheckFlags = 0, int64 SkipFlags = 0, const CustomExportCallback* ExportCb = nullptr, EJsonObjectConversionFlags ConversionFlags = EJsonObjectConversionFlags::None);
 
 	/**
 	 * Converts from a UStruct to a json string containing an object, using exportText
@@ -130,10 +139,10 @@ public: // UStruct -> JSON
 	 * Wrapper to UStructToJsonObjectString that allows a print policy to be specified.
 	 */
 	template<typename CharType, template<typename> class PrintPolicy>
-	static bool UStructToFormattedJsonObjectString(const UStruct* StructDefinition, const void* Struct, FString& OutJsonString, int64 CheckFlags = 0, int64 SkipFlags = 0, int32 Indent = 0, const CustomExportCallback* ExportCb = nullptr)
+	static bool UStructToFormattedJsonObjectString(const UStruct* StructDefinition, const void* Struct, FString& OutJsonString, int64 CheckFlags = 0, int64 SkipFlags = 0, int32 Indent = 0, const CustomExportCallback* ExportCb = nullptr, EJsonObjectConversionFlags ConversionFlags = EJsonObjectConversionFlags::None)
 	{
 		TSharedRef<FJsonObject> JsonObject = MakeShareable(new FJsonObject());
-		if (UStructToJsonObject(StructDefinition, Struct, JsonObject, CheckFlags, SkipFlags, ExportCb))
+		if (UStructToJsonObject(StructDefinition, Struct, JsonObject, CheckFlags, SkipFlags, ExportCb, ConversionFlags))
 		{
 			TSharedRef<TJsonWriter<CharType, PrintPolicy<CharType>>> JsonWriter = TJsonWriterFactory<CharType, PrintPolicy<CharType>>::Create(&OutJsonString, Indent);
 
@@ -161,10 +170,11 @@ public: // UStruct -> JSON
 	 * @param CheckFlags Only convert properties that match at least one of these flags. If 0 check all properties.
 	 * @param SkipFlags Skip properties that match any of these flags
 	 * @param ExportCb Optional callback to override export behavior, if this returns null it will fallback to the default
+	 * @param ConversionFlags Bitwise flags to customize the conversion behavior
 	 *
 	 * @return False if any properties failed to write
 	 */
-	static JSONUTILITIES_API bool UStructToJsonAttributes(const UStruct* StructDefinition, const void* Struct, TMap< FString, TSharedPtr<FJsonValue> >& OutJsonAttributes, int64 CheckFlags = 0, int64 SkipFlags = 0, const CustomExportCallback* ExportCb = nullptr);
+	static JSONUTILITIES_API bool UStructToJsonAttributes(const UStruct* StructDefinition, const void* Struct, TMap< FString, TSharedPtr<FJsonValue> >& OutJsonAttributes, int64 CheckFlags = 0, int64 SkipFlags = 0, const CustomExportCallback* ExportCb = nullptr, EJsonObjectConversionFlags ConversionFlags = EJsonObjectConversionFlags::None);
 
 	/* * Converts from a FProperty to a Json Value using exportText
 	 *
@@ -174,10 +184,11 @@ public: // UStruct -> JSON
 	 * @param SkipFlags			Skip properties that match any of these flags
 	 * @param ExportCb Optional callback to override export behavior, if this returns null it will fallback to the default
 	 * @param OuterProperty		If applicable, the Array/Set/Map Property that contains this property
+	 * @param ConversionFlags	Bitwise flags to customize the conversion behavior
 	 *
 	 * @return					The constructed JsonValue from the property
 	 */
-	static JSONUTILITIES_API TSharedPtr<FJsonValue> UPropertyToJsonValue(FProperty* Property, const void* Value, int64 CheckFlags = 0, int64 SkipFlags = 0, const CustomExportCallback* ExportCb = nullptr, FProperty* OuterProperty = nullptr);
+	static JSONUTILITIES_API TSharedPtr<FJsonValue> UPropertyToJsonValue(FProperty* Property, const void* Value, int64 CheckFlags = 0, int64 SkipFlags = 0, const CustomExportCallback* ExportCb = nullptr, FProperty* OuterProperty = nullptr, EJsonObjectConversionFlags ConversionFlags = EJsonObjectConversionFlags::None);
 
 public: // JSON -> UStruct
 
