JsonObjectConverter looks for meta field in UPROPERTY to serialize/deserialize json field names.

diff --git a/Engine/Source/Runtime/CoreUObject/Public/UObject/ObjectMacros.h b/Engine/Source/Runtime/CoreUObject/Public/UObject/ObjectMacros.h
--- a/Engine/Source/Runtime/CoreUObject/Public/UObject/ObjectMacros.h
+++ b/Engine/Source/Runtime/CoreUObject/Public/UObject/ObjectMacros.h
@@ -1479,6 +1479,9 @@ namespace UM
 
 	    /// [PropertyMetadata] Used for enum properties to define the subset of invalid values as a comma-separated string; values within this subset are not made available.
         InvalidEnumValues,
+
+		/// [PropertyMetadata] Used to override the JSON field name in serialization via FJsonObjectConverter.
+		JSONFieldName,
 	};
 
 	// Metadata usable in UPROPERTY for customizing the behavior of Persona and UMG
diff --git a/Engine/Source/Runtime/JsonUtilities/Private/JsonObjectConverter.cpp b/Engine/Source/Runtime/JsonUtilities/Private/JsonObjectConverter.cpp
--- a/Engine/Source/Runtime/JsonUtilities/Private/JsonObjectConverter.cpp
+++ b/Engine/Source/Runtime/JsonUtilities/Private/JsonObjectConverter.cpp
@@ -31,6 +31,8 @@ namespace
 
 	const FName NAME_DateTime(TEXT("DateTime"));
 
+	const FName NAME_JSONFieldName(TEXT("JSONFieldName"));
+
 /** Convert property to JSON, assuming either the property is not an array or the value is an individual array element */
 TSharedPtr<FJsonValue> ConvertScalarFPropertyToJsonValue(FProperty* Property, const void* Value, int64 CheckFlags, int64 SkipFlags, const FJsonObjectConverter::CustomExportCallback* ExportCb, FProperty* OuterProperty, EJsonObjectConversionFlags ConversionFlags)
 {
@@ -301,6 +303,13 @@ bool FJsonObjectConverter::UStructToJsonAttributes(const UStruct* StructDefiniti
 		}
 
 		FString VariableName = Property->GetAuthoredName();
+		if (const TMap<FName, FString>* PropMetadata = Property->GetMetaDataMap())
+		{
+			if (const FString* JsonFieldName = PropMetadata->Find(NAME_JSONFieldName))
+			{
+				VariableName = *JsonFieldName;
+			}
+		}
 		if (!EnumHasAnyFlags(ConversionFlags, EJsonObjectConversionFlags::SkipStandardizeCase))
 		{
 			VariableName = StandardizeCase(VariableName);
@@ -970,7 +979,16 @@ namespace
 			}
 
 			// find a JSON value matching this property name
+			// use JSONFieldName meta property, if available
 			FString PropertyName = StructDefinition->GetAuthoredNameForField(Property);
+			if (const TMap<FName, FString>* PropMetadata = Property->GetMetaDataMap())
+			{
+				if (const FString* JsonFieldName = PropMetadata->Find(NAME_JSONFieldName))
+				{
+					PropertyName = *JsonFieldName;
+				}
+			}
+			
 			const TSharedPtr<FJsonValue>* JsonValue = JsonAttributes.Find(PropertyName);
 			
 			if (!JsonValue)
