Fix non unity build for 28284565

diff --git a/Engine/Source/Runtime/JsonUtilities/Private/JsonObjectConverter.cpp b/Engine/Source/Runtime/JsonUtilities/Private/JsonObjectConverter.cpp
--- a/Engine/Source/Runtime/JsonUtilities/Private/JsonObjectConverter.cpp
+++ b/Engine/Source/Runtime/JsonUtilities/Private/JsonObjectConverter.cpp
@@ -10,6 +10,7 @@
 #include "UObject/TextProperty.h"
 #include "UObject/PropertyPortFlags.h"
 #include "UObject/Package.h"
+#include "UObject/Field.h"
 #include "Policies/CondensedJsonPrintPolicy.h"
 #include "JsonObjectWrapper.h"
 
