Backout JsonObjectConverter changes.

diff --git a/Engine/Source/Runtime/JsonUtilities/Private/JsonObjectConverter.cpp b/Engine/Source/Runtime/JsonUtilities/Private/JsonObjectConverter.cpp
--- a/Engine/Source/Runtime/JsonUtilities/Private/JsonObjectConverter.cpp
+++ b/Engine/Source/Runtime/JsonUtilities/Private/JsonObjectConverter.cpp
@@ -10,7 +10,6 @@
 #include "UObject/TextProperty.h"
 #include "UObject/PropertyPortFlags.h"
 #include "UObject/Package.h"
-#include "UObject/Field.h"
 #include "Policies/CondensedJsonPrintPolicy.h"
 #include "JsonObjectWrapper.h"
 
@@ -32,8 +31,6 @@ namespace
 
 	const FName NAME_DateTime(TEXT("DateTime"));
 
-	const FName NAME_JSONFieldName(TEXT("JSONFieldName"));
-
 /** Convert property to JSON, assuming either the property is not an array or the value is an individual array element */
 TSharedPtr<FJsonValue> ConvertScalarFPropertyToJsonValue(FProperty* Property, const void* Value, int64 CheckFlags, int64 SkipFlags, const FJsonObjectConverter::CustomExportCallback* ExportCb, FProperty* OuterProperty, EJsonObjectConversionFlags ConversionFlags)
 {
@@ -304,13 +301,6 @@ bool FJsonObjectConverter::UStructToJsonAttributes(const UStruct* StructDefiniti
 		}
 
 		FString VariableName = Property->GetAuthoredName();
-		if (const TMap<FName, FString>* PropMetadata = Property->GetMetaDataMap())
-		{
-			if (const FString* JsonFieldName = PropMetadata->Find(NAME_JSONFieldName))
-			{
-				VariableName = *JsonFieldName;
-			}
-		}
 		if (!EnumHasAnyFlags(ConversionFlags, EJsonObjectConversionFlags::SkipStandardizeCase))
 		{
 			VariableName = StandardizeCase(VariableName);
@@ -980,16 +970,7 @@ namespace
 			}
 
 			// find a JSON value matching this property name
-			// use JSONFieldName meta property, if available
 			FString PropertyName = StructDefinition->GetAuthoredNameForField(Property);
-			if (const TMap<FName, FString>* PropMetadata = Property->GetMetaDataMap())
-			{
-				if (const FString* JsonFieldName = PropMetadata->Find(NAME_JSONFieldName))
-				{
-					PropertyName = *JsonFieldName;
-				}
-			}
-			
 			const TSharedPtr<FJsonValue>* JsonValue = JsonAttributes.Find(PropertyName);
 			
 			if (!JsonValue)
diff --git a/Engine/Source/Runtime/JsonUtilities/Public/JsonObjectConverter.h b/Engine/Source/Runtime/JsonUtilities/Public/JsonObjectConverter.h
--- a/Engine/Source/Runtime/JsonUtilities/Public/JsonObjectConverter.h
+++ b/Engine/Source/Runtime/JsonUtilities/Public/JsonObjectConverter.h
@@ -22,7 +22,6 @@
 #include "Templates/SharedPointer.h"
 #include "Trace/Detail/Channel.h"
 #include "UObject/Class.h"
-#include "UObject/Field.h"
 
 enum class EJsonObjectConversionFlags
 {
