[json] DOM builder trivial: replaced enable_ifs with UE_REQUIRES
#rnx

diff --git a/Engine/Source/Runtime/JsonUtilities/Public/JsonDomBuilder.h b/Engine/Source/Runtime/JsonUtilities/Public/JsonDomBuilder.h
--- a/Engine/Source/Runtime/JsonUtilities/Public/JsonDomBuilder.h
+++ b/Engine/Source/Runtime/JsonUtilities/Public/JsonDomBuilder.h
@@ -2,14 +2,13 @@
 
 #pragma once
 
+#include <type_traits>
+
 #include "Dom/JsonValue.h"
 #include "Dom/JsonObject.h"
 #include "Serialization/JsonSerializer.h"
 
-#include "Templates/IsFloatingPoint.h"
-#include "Templates/IsIntegral.h"
-#include "Templates/EnableIf.h"
-#include "Templates/Invoke.h"
+#include "Templates/UnrealTypeTraits.h"
 
 /**
  * Helpers for creating TSharedPtr<FJsonValue> JSON trees
@@ -80,13 +79,13 @@ public:
 
 		FObject& Set(const FString& Key, const FString& Str)           { Object->SetField(Key, MakeShared<FJsonValueString>(Str));            return *this; }
 
-		template <class FNumber>
-		typename TEnableIf<!std::is_same_v<FNumber, bool> && (TIsIntegral<FNumber>::Value || TIsFloatingPoint<FNumber>::Value), FObject&>::Type
-			Set(const FString& Key, FNumber Number)                    { Object->SetField(Key, MakeShared<FJsonValueNumber>(Number));         return *this; }
+		template <class NumberType
+			UE_REQUIRES(!std::is_same_v<NumberType, bool> && (std::is_integral_v<NumberType> || std::is_floating_point_v<NumberType>))>
+		FObject& Set(const FString& Key, NumberType Number)                    { Object->SetField(Key, MakeShared<FJsonValueNumber>(Number));         return *this; }
 
-		template <class FBool>
-		typename TEnableIf<std::is_same_v<FBool, bool>, FObject&>::Type
-			Set(const FString& Key, FBool Boolean)                     { Object->SetField(Key, MakeShared<FJsonValueBoolean>(Boolean));       return *this; }
+		template <class BoolType
+			UE_REQUIRES(std::is_same_v<BoolType, bool>)>
+		FObject& Set(const FString& Key, BoolType Boolean)                     { Object->SetField(Key, MakeShared<FJsonValueBoolean>(Boolean));       return *this; }
 
 		FObject& Set(const FString& Key, TYPE_OF_NULLPTR)              { Object->SetField(Key, MakeShared<FJsonValueNull>());                 return *this; }
 
@@ -144,12 +143,12 @@ public:
 		FArray& Add(TSharedPtr<FJsonValue> Value) { Array.Emplace(Value);                                  return *this; }
 
 		/** Add multiple values */
-		template <class... FValue>
-		typename TEnableIf<(sizeof...(FValue) > 1), FArray&>::Type
-			Add(FValue&&... Value)
+		template <class... ValueType
+			UE_REQUIRES(sizeof...(ValueType) > 1)>
+		FArray& Add(ValueType&&... Value)
 		{
 			// This should be implemented with a fold expression when our compilers support it
-			int Temp[] = {0, (Add(Forward<FValue>(Value)), 0)...};
+			int Temp[] = {0, (Add(Forward<ValueType>(Value)), 0)...};
 			(void)Temp;
 			return *this;
 		}
