Allow Json object helpers to work with UClasses as well as UScriptStructs

diff --git a/Engine/Source/Runtime/JsonUtilities/Public/JsonObjectConverter.h b/Engine/Source/Runtime/JsonUtilities/Public/JsonObjectConverter.h
--- a/Engine/Source/Runtime/JsonUtilities/Public/JsonObjectConverter.h
+++ b/Engine/Source/Runtime/JsonUtilities/Public/JsonObjectConverter.h
@@ -22,6 +22,8 @@
 #include "Templates/SharedPointer.h"
 #include "Trace/Detail/Channel.h"
 #include "UObject/Class.h"
+#include "Templates/Models.h"
+#include "Concepts/StaticClassProvider.h"
 
 enum class EJsonObjectConversionFlags
 {
@@ -132,7 +134,14 @@ public: // UStruct -> JSON
 	template<typename InStructType>
 	static bool UStructToJsonObjectString(const InStructType& InStruct, FString& OutJsonString, int64 CheckFlags = 0, int64 SkipFlags = 0, int32 Indent = 0, const CustomExportCallback* ExportCb = nullptr, bool bPrettyPrint = true)
 	{
-		return UStructToJsonObjectString(InStructType::StaticStruct(), &InStruct, OutJsonString, CheckFlags, SkipFlags, Indent, ExportCb, bPrettyPrint);
+		if constexpr (TModels<CStaticClassProvider, InStructType>::Value)
+		{
+			return UStructToJsonObjectString(InStructType::StaticClass(), &InStruct, OutJsonString, CheckFlags, SkipFlags, Indent, ExportCb, bPrettyPrint);
+		}
+		else
+		{
+			return UStructToJsonObjectString(InStructType::StaticStruct(), &InStruct, OutJsonString, CheckFlags, SkipFlags, Indent, ExportCb, bPrettyPrint);
+		}
 	}
 
 	/**
@@ -222,7 +231,14 @@ public: // JSON -> UStruct
 	template<typename OutStructType>
 	static bool JsonObjectToUStruct(const TSharedRef<FJsonObject>& JsonObject, OutStructType* OutStruct, int64 CheckFlags = 0, int64 SkipFlags = 0, const bool bStrictMode = false, FText* OutFailReason = nullptr)
 	{
-		return JsonObjectToUStruct(JsonObject, OutStructType::StaticStruct(), OutStruct, CheckFlags, SkipFlags, bStrictMode, OutFailReason);
+		if constexpr (TModels<CStaticClassProvider, OutStructType>::Value)
+		{
+			return JsonObjectToUStruct(JsonObject, OutStructType::StaticClass(), OutStruct, CheckFlags, SkipFlags, bStrictMode, OutFailReason);
+		}
+		else
+		{
+			return JsonObjectToUStruct(JsonObject, OutStructType::StaticStruct(), OutStruct, CheckFlags, SkipFlags, bStrictMode, OutFailReason);
+		}
 	}
 
 	/**
