Add int64 support to JsonDomBuilder

diff --git a/Engine/Source/Runtime/JsonUtilities/Public/JsonDomBuilder.h b/Engine/Source/Runtime/JsonUtilities/Public/JsonDomBuilder.h
--- a/Engine/Source/Runtime/JsonUtilities/Public/JsonDomBuilder.h
+++ b/Engine/Source/Runtime/JsonUtilities/Public/JsonDomBuilder.h
@@ -81,7 +81,9 @@ public:
 
 		template <class NumberType
 			UE_REQUIRES(!std::is_same_v<NumberType, bool> && (std::is_integral_v<NumberType> || std::is_floating_point_v<NumberType>))>
-		FObject& Set(const FString& Key, NumberType Number)                    { Object->SetField(Key, MakeShared<FJsonValueNumber>(Number));         return *this; }
+		FObject& Set(const FString& Key, NumberType Number)            { Object->SetField(Key, MakeShared<FJsonValueNumber>(Number));         return *this; }
+
+		FObject& Set(const FString& Key, const int64& Number)          { Object->SetField(Key, MakeShared<FJsonValueNumberString>(FString::Printf(TEXT("%" INT64_FMT), Number))); return *this; }
 
 		template <class BoolType
 			UE_REQUIRES(std::is_same_v<BoolType, bool>)>
