Enable AutoRTFM for json modules

diff --git a/Engine/Source/Runtime/Json/Json.Build.cs b/Engine/Source/Runtime/Json/Json.Build.cs
--- a/Engine/Source/Runtime/Json/Json.Build.cs
+++ b/Engine/Source/Runtime/Json/Json.Build.cs
@@ -15,6 +15,8 @@ namespace UnrealBuildTool.Rules
 			); 
 
 			UnsafeTypeCastWarningLevel = WarningLevel.Error;
+
+			bAllowAutoRTFMInstrumentation = true;
 		}
 	}
 }
diff --git a/Engine/Source/Runtime/JsonUtilities/JsonUtilities.Build.cs b/Engine/Source/Runtime/JsonUtilities/JsonUtilities.Build.cs
--- a/Engine/Source/Runtime/JsonUtilities/JsonUtilities.Build.cs
+++ b/Engine/Source/Runtime/JsonUtilities/JsonUtilities.Build.cs
@@ -14,5 +14,7 @@ public class JsonUtilities : ModuleRules
 				"Json",
 			}
 		);
+
+		bAllowAutoRTFMInstrumentation = true;
 	}
 }
