Added missing OutFailReason parameters to some FJsonObjectConverter functions.

diff --git a/Engine/Source/Runtime/JsonUtilities/Public/JsonObjectConverter.h b/Engine/Source/Runtime/JsonUtilities/Public/JsonObjectConverter.h
--- a/Engine/Source/Runtime/JsonUtilities/Public/JsonObjectConverter.h
+++ b/Engine/Source/Runtime/JsonUtilities/Public/JsonObjectConverter.h
@@ -25,6 +25,8 @@
 #include "Templates/Models.h"
 #include "Concepts/StaticClassProvider.h"
 
+#define LOCTEXT_NAMESPACE "JsonObjectConverter"
+
 enum class EJsonObjectConversionFlags
 {
 	None = 0,
@@ -283,18 +285,26 @@ public: // JSON -> UStruct
 	 * @return False if any properties matched but failed to deserialize
 	 */
 	template<typename OutStructType>
-	static bool JsonObjectStringToUStruct(const FString& JsonString, OutStructType* OutStruct, int64 CheckFlags = 0, int64 SkipFlags = 0, const bool bStrictMode = false)
+	static bool JsonObjectStringToUStruct(const FString& JsonString, OutStructType* OutStruct, int64 CheckFlags = 0, int64 SkipFlags = 0, const bool bStrictMode = false, FText* OutFailReason = nullptr)
 	{
 		TSharedPtr<FJsonObject> JsonObject;
 		TSharedRef<TJsonReader<> > JsonReader = TJsonReaderFactory<>::Create(JsonString);
 		if (!FJsonSerializer::Deserialize(JsonReader, JsonObject) || !JsonObject.IsValid())
 		{
-			UE_LOG(LogJson, Warning, TEXT("JsonObjectStringToUStruct - Unable to parse json=[%s]"), *JsonString);
+			UE_LOG(LogJson, Warning, TEXT("JsonObjectStringToUStruct - Unable to parse. json=[%s]"), *JsonString);
+			if (OutFailReason)
+			{
+				*OutFailReason = FText::Format(LOCTEXT("FailJsonObjectDeserialize", "JsonObjectStringToUStruct - Unable to parse. json=[{0}]"), FText::FromString(*JsonString));
+			}
 			return false;
 		}
-		if (!FJsonObjectConverter::JsonObjectToUStruct(JsonObject.ToSharedRef(), OutStruct, CheckFlags, SkipFlags, bStrictMode))
+		if (!FJsonObjectConverter::JsonObjectToUStruct(JsonObject.ToSharedRef(), OutStruct, CheckFlags, SkipFlags, bStrictMode, OutFailReason))
 		{
 			UE_LOG(LogJson, Warning, TEXT("JsonObjectStringToUStruct - Unable to deserialize. json=[%s]"), *JsonString);
+			if (OutFailReason)
+			{
+				*OutFailReason = FText::Format(LOCTEXT("FailJsonObjectConversion", "JsonObjectStringToUStruct - Unable to deserialize. json=[{0}]\n{1}"), FText::FromString(*JsonString), *OutFailReason);
+			}
 			return false;
 		}
 		return true;
@@ -312,18 +322,26 @@ public: // JSON -> UStruct
 	* @return False if any properties matched but failed to deserialize.
 	*/
 	template<typename OutStructType>
-	static bool JsonArrayStringToUStruct(const FString& JsonString, TArray<OutStructType>* OutStructArray, int64 CheckFlags = 0, int64 SkipFlags = 0, const bool bStrictMode = false)
+	static bool JsonArrayStringToUStruct(const FString& JsonString, TArray<OutStructType>* OutStructArray, int64 CheckFlags = 0, int64 SkipFlags = 0, const bool bStrictMode = false, FText* OutFailReason = nullptr)
 	{
 		TArray<TSharedPtr<FJsonValue> > JsonArray;
 		TSharedRef<TJsonReader<> > JsonReader = TJsonReaderFactory<>::Create(JsonString);
 		if (!FJsonSerializer::Deserialize(JsonReader, JsonArray))
 		{
 			UE_LOG(LogJson, Warning, TEXT("JsonArrayStringToUStruct - Unable to parse. json=[%s]"), *JsonString);
+			if (OutFailReason)
+			{
+				*OutFailReason = FText::Format(LOCTEXT("FailJsonArrayDeserialize", "JsonArrayStringToUStruct - Unable to parse. json=[{0}]"), FText::FromString(*JsonString));
+			}
 			return false;
 		}
-		if (!JsonArrayToUStruct(JsonArray, OutStructArray, CheckFlags, SkipFlags, bStrictMode))
+		if (!JsonArrayToUStruct(JsonArray, OutStructArray, CheckFlags, SkipFlags, bStrictMode, OutFailReason))
 		{
 			UE_LOG(LogJson, Warning, TEXT("JsonArrayStringToUStruct - Error parsing one of the elements. json=[%s]"), *JsonString);
+			if (OutFailReason)
+			{
+				*OutFailReason = FText::Format(LOCTEXT("FailJsonArrayConversion", "JsonArrayStringToUStruct - Error parsing one of the elements. json=[{0}]\n{1}"), FText::FromString(*JsonString), *OutFailReason);
+			}
 			return false;
 		}
 		return true;
@@ -341,7 +359,7 @@ public: // JSON -> UStruct
 	* @return False if any of the matching elements are not an object, or if one of the matching elements could not be converted to the specified UStruct type.
 	*/
 	template<typename OutStructType>
-	static bool JsonArrayToUStruct(const TArray<TSharedPtr<FJsonValue>>& JsonArray, TArray<OutStructType>* OutStructArray, int64 CheckFlags = 0, int64 SkipFlags = 0, const bool bStrictMode = false)
+	static bool JsonArrayToUStruct(const TArray<TSharedPtr<FJsonValue>>& JsonArray, TArray<OutStructType>* OutStructArray, int64 CheckFlags = 0, int64 SkipFlags = 0, const bool bStrictMode = false, FText* OutFailReason = nullptr)
 	{
 		OutStructArray->SetNum(JsonArray.Num());
 		for (int32 i = 0; i < JsonArray.Num(); ++i)
@@ -350,11 +368,19 @@ public: // JSON -> UStruct
 			if (Value->Type != EJson::Object)
 			{
 				UE_LOG(LogJson, Warning, TEXT("JsonArrayToUStruct - Array element [%i] was not an object."), i);
+				if (OutFailReason)
+				{
+					*OutFailReason = FText::Format(LOCTEXT("FailJsonArrayElementObject", "JsonArrayToUStruct - Array element [{0}] was not an object."), i);
+				}
 				return false;
 			}
-			if (!FJsonObjectConverter::JsonObjectToUStruct(Value->AsObject().ToSharedRef(), OutStructType::StaticStruct(), &(*OutStructArray)[i], CheckFlags, SkipFlags, bStrictMode))
+			if (!FJsonObjectConverter::JsonObjectToUStruct(Value->AsObject().ToSharedRef(), OutStructType::StaticStruct(), &(*OutStructArray)[i], CheckFlags, SkipFlags, bStrictMode, OutFailReason))
 			{
 				UE_LOG(LogJson, Warning, TEXT("JsonArrayToUStruct - Unable to convert element [%i]."), i);
+				if (OutFailReason)
+				{
+					*OutFailReason = FText::Format(LOCTEXT("FailJsonArrayElementConversion", "JsonArrayToUStruct - Unable to convert element [{0}].\n{1}"), i, *OutFailReason);
+				}
 				return false;
 			}
 		}
@@ -367,3 +393,5 @@ public: // JSON -> UStruct
 	*/
 	static JSONUTILITIES_API FFormatNamedArguments ParseTextArgumentsFromJson(const TSharedPtr<const FJsonObject>& JsonObject);
 };
+
+#undef LOCTEXT_NAMESPACE
