Removed deprecated function and added missing OutFailReason parameter comments.

diff --git a/Engine/Source/Runtime/JsonUtilities/Private/JsonObjectConverter.cpp b/Engine/Source/Runtime/JsonUtilities/Private/JsonObjectConverter.cpp
--- a/Engine/Source/Runtime/JsonUtilities/Private/JsonObjectConverter.cpp
+++ b/Engine/Source/Runtime/JsonUtilities/Private/JsonObjectConverter.cpp
@@ -208,34 +208,6 @@ TSharedPtr<FJsonValue> ConvertScalarFPropertyToJsonValue(FProperty* Property, co
 }
 }
 
-PRAGMA_DISABLE_DEPRECATION_WARNINGS
-
-TSharedPtr<FJsonValue> FJsonObjectConverter::ObjectJsonCallback(FProperty* Property, const void* Value)
-{
-	if (FObjectProperty* ObjectProperty = CastField<FObjectProperty>(Property))
-	{
-		if (!ObjectProperty->HasAnyFlags(RF_Transient)) // We are taking Transient to mean we don't want to serialize to JSON either (could make a new flag if necessary)
-		{
-			TSharedRef<FJsonObject> Out = MakeShared<FJsonObject>();
-
-			CustomExportCallback CustomCB;
-			CustomCB.BindStatic(FJsonObjectConverter::ObjectJsonCallback);
-
-			void** PtrToValuePtr = (void**)Value;
-
-			if (FJsonObjectConverter::UStructToJsonObject(ObjectProperty->PropertyClass, (*PtrToValuePtr), Out, 0, 0, &CustomCB)) //-V549
-			{
-				return MakeShared<FJsonValueObject>(Out);
-			}
-		}
-	}
-
-	// invalid
-	return TSharedPtr<FJsonValue>();
-}
-
-PRAGMA_ENABLE_DEPRECATION_WARNINGS
-
 TSharedPtr<FJsonValue> FJsonObjectConverter::UPropertyToJsonValue(FProperty* Property, const void* Value, int64 CheckFlags, int64 SkipFlags, const CustomExportCallback* ExportCb, FProperty* OuterProperty, EJsonObjectConversionFlags ConversionFlags)
 {
 	if (Property->ArrayDim == 1)
diff --git a/Engine/Source/Runtime/JsonUtilities/Public/JsonObjectConverter.h b/Engine/Source/Runtime/JsonUtilities/Public/JsonObjectConverter.h
--- a/Engine/Source/Runtime/JsonUtilities/Public/JsonObjectConverter.h
+++ b/Engine/Source/Runtime/JsonUtilities/Public/JsonObjectConverter.h
@@ -63,12 +63,6 @@ public: // UStruct -> JSON
 
 	static JSONUTILITIES_API const CustomExportCallback ExportCallback_WriteISO8601Dates;
 
-	/**
-	 * Utility Export Callback for having object properties expanded to full Json.
-	 */
-	UE_DEPRECATED(4.25, "ObjectJsonCallback has been deprecated - please remove the usage of it from your project")
-	static JSONUTILITIES_API TSharedPtr<FJsonValue> ObjectJsonCallback(FProperty* Property , const void* Value);
-
 	/**
 	 * Templated version of UStructToJsonObject to try and make most of the params. Also serves as an example use case
 	 *
@@ -281,6 +275,7 @@ public: // JSON -> UStruct
 	 * @param CheckFlags Only convert properties that match at least one of these flags. If 0 check all properties.
 	 * @param SkipFlags Skip properties that match any of these flags
 	 * @param bStrictMode Whether to strictly check the json attributes
+	 * @param OutFailReason Reason of the failure if any
 	 *
 	 * @return False if any properties matched but failed to deserialize
 	 */
@@ -311,16 +306,17 @@ public: // JSON -> UStruct
 	}
 
 	/**
-	* Converts from a json string containing an array to an array of UStructs
-	*
-	* @param JsonString String containing JSON formatted data.
-	* @param OutStructArray The UStruct array to copy in to
-	* @param CheckFlags Only convert properties that match at least one of these flags. If 0 check all properties.
-	* @param SkipFlags Skip properties that match any of these flags.
-	* @param bStrictMode Whether to strictly check the json attributes
-	*
-	* @return False if any properties matched but failed to deserialize.
-	*/
+	 * Converts from a json string containing an array to an array of UStructs
+	 *
+	 * @param JsonString String containing JSON formatted data.
+	 * @param OutStructArray The UStruct array to copy in to
+	 * @param CheckFlags Only convert properties that match at least one of these flags. If 0 check all properties.
+	 * @param SkipFlags Skip properties that match any of these flags.
+	 * @param bStrictMode Whether to strictly check the json attributes
+	 * @param OutFailReason Reason of the failure if any
+	 *
+	 * @return False if any properties matched but failed to deserialize.
+	 */
 	template<typename OutStructType>
 	static bool JsonArrayStringToUStruct(const FString& JsonString, TArray<OutStructType>* OutStructArray, int64 CheckFlags = 0, int64 SkipFlags = 0, const bool bStrictMode = false, FText* OutFailReason = nullptr)
 	{
@@ -348,16 +344,17 @@ public: // JSON -> UStruct
 	}
 
 	/**
-	* Converts from an array of json values to an array of UStructs.
-	*
-	* @param JsonArray Array containing json values to convert.
-	* @param OutStructArray The UStruct array to copy in to
-	* @param CheckFlags Only convert properties that match at least one of these flags. If 0 check all properties.
-	* @param SkipFlags Skip properties that match any of these flags.
-	* @param bStrictMode Whether to strictly check the json attributes
-	*
-	* @return False if any of the matching elements are not an object, or if one of the matching elements could not be converted to the specified UStruct type.
-	*/
+	 * Converts from an array of json values to an array of UStructs.
+	 *
+	 * @param JsonArray Array containing json values to convert.
+	 * @param OutStructArray The UStruct array to copy in to
+	 * @param CheckFlags Only convert properties that match at least one of these flags. If 0 check all properties.
+	 * @param SkipFlags Skip properties that match any of these flags.
+	 * @param bStrictMode Whether to strictly check the json attributes
+	 * @param OutFailReason Reason of the failure if any
+	 *
+	 * @return False if any of the matching elements are not an object, or if one of the matching elements could not be converted to the specified UStruct type.
+	 */
 	template<typename OutStructType>
 	static bool JsonArrayToUStruct(const TArray<TSharedPtr<FJsonValue>>& JsonArray, TArray<OutStructType>* OutStructArray, int64 CheckFlags = 0, int64 SkipFlags = 0, const bool bStrictMode = false, FText* OutFailReason = nullptr)
 	{
