GitHub #11170 : Add a custom import function to JsonObjectConverter methods to match inverse of existing custom export callback

diff --git a/Engine/Source/Runtime/JsonUtilities/Private/JsonObjectConverter.cpp b/Engine/Source/Runtime/JsonUtilities/Private/JsonObjectConverter.cpp
--- a/Engine/Source/Runtime/JsonUtilities/Private/JsonObjectConverter.cpp
+++ b/Engine/Source/Runtime/JsonUtilities/Private/JsonObjectConverter.cpp
@@ -369,382 +369,382 @@ bool FJsonObjectConverter::GetTextFromObject(const TSharedRef<FJsonObject>& Obj,
 
 namespace
 {
-	bool JsonValueToFPropertyWithContainer(const TSharedPtr<FJsonValue>& JsonValue, FProperty* Property, void* OutValue, const UStruct* ContainerStruct, void* Container, int64 CheckFlags, int64 SkipFlags, const bool bStrictMode, FText* OutFailReason);
-	bool JsonAttributesToUStructWithContainer(const TMap< FString, TSharedPtr<FJsonValue> >& JsonAttributes, const UStruct* StructDefinition, void* OutStruct, const UStruct* ContainerStruct, void* Container, int64 CheckFlags, int64 SkipFlags, const bool bStrictMode, FText* OutFailReason);
+	bool JsonValueToFPropertyWithContainer(const TSharedPtr<FJsonValue>& JsonValue, FProperty* Property, void* OutValue, const UStruct* ContainerStruct, void* Container, int64 CheckFlags, int64 SkipFlags, const bool bStrictMode, FText* OutFailReason, const FJsonObjectConverter::CustomImportCallback* ImportCb);
+	bool JsonAttributesToUStructWithContainer(const TMap< FString, TSharedPtr<FJsonValue> >& JsonAttributes, const UStruct* StructDefinition, void* OutStruct, const UStruct* ContainerStruct, void* Container, int64 CheckFlags, int64 SkipFlags, const bool bStrictMode, FText* OutFailReason, const FJsonObjectConverter::CustomImportCallback* ImportCb);
 
 	/** Convert JSON to property, assuming either the property is not an array or the value is an individual array element */
-	bool ConvertScalarJsonValueToFPropertyWithContainer(const TSharedPtr<FJsonValue>& JsonValue, FProperty* Property, void* OutValue, const UStruct* ContainerStruct, void* Container, int64 CheckFlags, int64 SkipFlags, const bool bStrictMode, FText* OutFailReason)
+	bool ConvertScalarJsonValueToFPropertyWithContainer(const TSharedPtr<FJsonValue>& JsonValue, FProperty* Property, void* OutValue, const UStruct* ContainerStruct, void* Container, int64 CheckFlags, int64 SkipFlags, const bool bStrictMode, FText* OutFailReason, const FJsonObjectConverter::CustomImportCallback* ImportCb)
 	{
-	if (FEnumProperty* EnumProperty = CastField<FEnumProperty>(Property))
-	{
-		if (JsonValue->Type == EJson::String)
-		{
-			// see if we were passed a string for the enum
-			const UEnum* Enum = EnumProperty->GetEnum();
-			check(Enum);
-			FString StrValue = JsonValue->AsString();
-			int64 IntValue = Enum->GetValueByName(FName(*StrValue), EGetByNameFlags::CheckAuthoredName);
-			if (IntValue == INDEX_NONE)
-			{
-				UE_LOG(LogJson, Error, TEXT("JsonValueToUProperty - Unable to import enum %s from string value %s for property %s"), *Enum->CppType, *StrValue, *Property->GetAuthoredName());
-				if (OutFailReason)
-				{
-					*OutFailReason = FText::Format(LOCTEXT("FailImportEnumFromString", "Unable to import enum {0} from string value {1} for property {2}"), FText::FromString(Enum->CppType), FText::FromString(StrValue), FText::FromString(Property->GetAuthoredName()));
-				}
-				return false;
-			}
-			EnumProperty->GetUnderlyingProperty()->SetIntPropertyValue(OutValue, IntValue);
-		}
-		else
-		{
-			// AsNumber will log an error for completely inappropriate types (then give us a default)
-			EnumProperty->GetUnderlyingProperty()->SetIntPropertyValue(OutValue, (int64)JsonValue->AsNumber());
-		}
-	}
-	else if (FNumericProperty *NumericProperty = CastField<FNumericProperty>(Property))
-	{
-		if (NumericProperty->IsEnum() && JsonValue->Type == EJson::String)
-		{
-			// see if we were passed a string for the enum
-			const UEnum* Enum = NumericProperty->GetIntPropertyEnum();
-			check(Enum); // should be assured by IsEnum()
-			FString StrValue = JsonValue->AsString();
-			int64 IntValue = Enum->GetValueByName(FName(*StrValue), EGetByNameFlags::CheckAuthoredName);
-			if (IntValue == INDEX_NONE)
-			{
-				UE_LOG(LogJson, Error, TEXT("JsonValueToUProperty - Unable to import enum %s from numeric value %s for property %s"), *Enum->CppType, *StrValue, *Property->GetAuthoredName());
-				if (OutFailReason)
-				{
-					*OutFailReason = FText::Format(LOCTEXT("FailImportEnumFromNumeric", "Unable to import enum {0} from numeric value {1} for property {2}"), FText::FromString(Enum->CppType), FText::FromString(StrValue), FText::FromString(Property->GetAuthoredName()));
-				}
-				return false;
-			}
-			NumericProperty->SetIntPropertyValue(OutValue, IntValue);
-		}
-		else if (NumericProperty->IsFloatingPoint())
-		{
-			// AsNumber will log an error for completely inappropriate types (then give us a default)
-			NumericProperty->SetFloatingPointPropertyValue(OutValue, JsonValue->AsNumber());
-		}
-		else if (NumericProperty->IsInteger())
-		{
-			if (JsonValue->Type == EJson::String)
-			{
-				// parse string -> int64 ourselves so we don't lose any precision going through AsNumber (aka double)
-				NumericProperty->SetIntPropertyValue(OutValue, FCString::Atoi64(*JsonValue->AsString()));
-			}
-			else
-			{
-				// AsNumber will log an error for completely inappropriate types (then give us a default)
-				NumericProperty->SetIntPropertyValue(OutValue, (int64)JsonValue->AsNumber());
-			}
-		}
-		else
-		{
-			UE_LOG(LogJson, Error, TEXT("JsonValueToUProperty - Unable to import json value into %s numeric property %s"), *Property->GetClass()->GetName(), *Property->GetAuthoredName());
-			if (OutFailReason)
-			{
-				*OutFailReason = FText::Format(LOCTEXT("FailImportNumericProperty", "Unable to import json value into {0} numeric property {1}"), FText::FromString(Property->GetClass()->GetName()), FText::FromString(Property->GetAuthoredName()));
-			}
-			return false;
-		}
-	}
-	else if (FBoolProperty *BoolProperty = CastField<FBoolProperty>(Property))
-	{
-		// AsBool will log an error for completely inappropriate types (then give us a default)
-		BoolProperty->SetPropertyValue(OutValue, JsonValue->AsBool());
-	}
-	else if (FStrProperty *StringProperty = CastField<FStrProperty>(Property))
-	{
-		// AsString will log an error for completely inappropriate types (then give us a default)
-		StringProperty->SetPropertyValue(OutValue, JsonValue->AsString());
-	}
-	else if (FArrayProperty *ArrayProperty = CastField<FArrayProperty>(Property))
-	{
-		if (JsonValue->Type == EJson::Array)
-		{
-			TArray< TSharedPtr<FJsonValue> > ArrayValue = JsonValue->AsArray();
-			int32 ArrLen = ArrayValue.Num();
-
-			// make the output array size match
-			FScriptArrayHelper Helper(ArrayProperty, OutValue);
-			Helper.Resize(ArrLen);
-
-			// set the property values
-			for (int32 i = 0; i < ArrLen; ++i)
-			{
-				const TSharedPtr<FJsonValue>& ArrayValueItem = ArrayValue[i];
-				if (ArrayValueItem.IsValid() && !ArrayValueItem->IsNull())
-				{
-					if (!JsonValueToFPropertyWithContainer(ArrayValueItem, ArrayProperty->Inner, Helper.GetRawPtr(i), ContainerStruct, Container, CheckFlags & (~CPF_ParmFlags), SkipFlags, bStrictMode, OutFailReason))
-					{
-						UE_LOG(LogJson, Error, TEXT("JsonValueToUProperty - Unable to import Array element %d for property %s"), i, *Property->GetAuthoredName());
-						if (OutFailReason)
-						{
-							*OutFailReason = FText::Format(LOCTEXT("FailImportArrayElement", "Unable to import Array element {0} for property {1}\n{2}"), FText::AsNumber(i), FText::FromString(Property->GetAuthoredName()), *OutFailReason);
-						}
-						return false;
-					}
-				}
-			}
-		}
-		else
-		{
-			UE_LOG(LogJson, Error, TEXT("JsonValueToUProperty - Unable to import non-array JSON value into Array property %s"), *Property->GetAuthoredName());
-			if (OutFailReason)
-			{
-				*OutFailReason = FText::Format(LOCTEXT("FailImportArray", "Unable to import non-array JSON value into Array property {0}"), FText::FromString(Property->GetAuthoredName()));
-			}
-			return false;
-		}
-	}
-	else if (FMapProperty* MapProperty = CastField<FMapProperty>(Property))
-	{
-		if (JsonValue->Type == EJson::Object)
-		{
-			TSharedPtr<FJsonObject> ObjectValue = JsonValue->AsObject();
-
-			FScriptMapHelper Helper(MapProperty, OutValue);
-
-			check(ObjectValue);
-
-			int32 MapSize = ObjectValue->Values.Num();
-			Helper.EmptyValues(MapSize);
-
-			// set the property values
-			for (const auto& Entry : ObjectValue->Values)
-			{
-				if (Entry.Value.IsValid() && !Entry.Value->IsNull())
-				{
-					int32 NewIndex = Helper.AddDefaultValue_Invalid_NeedsRehash();
-
-					TSharedPtr<FJsonValueString> TempKeyValue = MakeShared<FJsonValueString>(Entry.Key);
-
-					if (!JsonValueToFPropertyWithContainer(TempKeyValue, MapProperty->KeyProp, Helper.GetKeyPtr(NewIndex), ContainerStruct, Container, CheckFlags & (~CPF_ParmFlags), SkipFlags, bStrictMode, OutFailReason))
-					{
-						UE_LOG(LogJson, Error, TEXT("JsonValueToUProperty - Unable to import Map element %s key for property %s"), *Entry.Key, *Property->GetAuthoredName());
-						if (OutFailReason)
-						{
-							*OutFailReason = FText::Format(LOCTEXT("FailImportMapElementKey", "Unable to import Map element {0} key for property {1}\n{2}"), FText::FromString(Entry.Key), FText::FromString(Property->GetAuthoredName()), *OutFailReason);
-						}
-						return false;
-					}
-
-					if (!JsonValueToFPropertyWithContainer(Entry.Value, MapProperty->ValueProp, Helper.GetValuePtr(NewIndex), ContainerStruct, Container, CheckFlags & (~CPF_ParmFlags), SkipFlags, bStrictMode, OutFailReason))
-					{
-						UE_LOG(LogJson, Error, TEXT("JsonValueToUProperty - Unable to import Map element %s value for property %s"), *Entry.Key, *Property->GetAuthoredName());
-						if (OutFailReason)
-						{
-							*OutFailReason = FText::Format(LOCTEXT("FailImportMapElementValue", "Unable to import Map element {0} value for property {1}\n{2}"), FText::FromString(Entry.Key), FText::FromString(Property->GetAuthoredName()), *OutFailReason);
-						}
-						return false;
-					}
-				}
-			}
-
-			Helper.Rehash();
-		}
-		else
-		{
-			UE_LOG(LogJson, Error, TEXT("JsonValueToUProperty - Unable to import non-object JSON value into Map property %s"), *Property->GetAuthoredName());
-			if (OutFailReason)
-			{
-				*OutFailReason = FText::Format(LOCTEXT("FailImportMap", "Unable to import non-object JSON value into Map property {0}"), FText::FromString(Property->GetAuthoredName()));
-			}
-			return false;
-		}
-	}
-	else if (FSetProperty* SetProperty = CastField<FSetProperty>(Property))
-	{
-		if (JsonValue->Type == EJson::Array)
-		{
-			TArray< TSharedPtr<FJsonValue> > ArrayValue = JsonValue->AsArray();
-			int32 ArrLen = ArrayValue.Num();
-
-			FScriptSetHelper Helper(SetProperty, OutValue);
-			Helper.EmptyElements(ArrLen);
-
-			// set the property values
-			for (int32 i = 0; i < ArrLen; ++i)
-			{
-				const TSharedPtr<FJsonValue>& ArrayValueItem = ArrayValue[i];
-				if (ArrayValueItem.IsValid() && !ArrayValueItem->IsNull())
-				{
-					int32 NewIndex = Helper.AddDefaultValue_Invalid_NeedsRehash();
-					if (!JsonValueToFPropertyWithContainer(ArrayValueItem, SetProperty->ElementProp, Helper.GetElementPtr(NewIndex), ContainerStruct, Container, CheckFlags & (~CPF_ParmFlags), SkipFlags, bStrictMode, OutFailReason))
-					{
-						UE_LOG(LogJson, Error, TEXT("JsonValueToUProperty - Unable to import Set element %d for property %s"), i, *Property->GetAuthoredName());
-						if (OutFailReason)
-						{
-							*OutFailReason = FText::Format(LOCTEXT("FailImportSetElement", "Unable to import Set element {0} for property {1}\n{2}"), FText::AsNumber(i), FText::FromString(Property->GetAuthoredName()), *OutFailReason);
-						}
-						return false;
-					}
-				}
-			}
-
-			Helper.Rehash();
-		}
-		else
-		{
-			UE_LOG(LogJson, Error, TEXT("JsonValueToUProperty - Unable to import non-array JSON value into Set property %s"), *Property->GetAuthoredName());
-			if (OutFailReason)
-			{
-				*OutFailReason = FText::Format(LOCTEXT("FailImportSet", "Unable to import non-array JSON value into Set property {0}"), FText::FromString(Property->GetAuthoredName()));
-			}
-			return false;
-		}
-	}
-	else if (FTextProperty* TextProperty = CastField<FTextProperty>(Property))
-	{
-		if (JsonValue->Type == EJson::String)
-		{
-			// assume this string is already localized, so import as invariant
-			TextProperty->SetPropertyValue(OutValue, FText::FromString(JsonValue->AsString()));
-		}
-		else if (JsonValue->Type == EJson::Object)
-		{
-			TSharedPtr<FJsonObject> Obj = JsonValue->AsObject();
-			check(Obj.IsValid()); // should not fail if Type == EJson::Object
-
-			// import the subvalue as a culture invariant string
-			FText Text;
-			if (!FJsonObjectConverter::GetTextFromObject(Obj.ToSharedRef(), Text))
-			{
-				UE_LOG(LogJson, Error, TEXT("JsonValueToUProperty - Unable to import JSON object with invalid keys into Text property %s"), *Property->GetAuthoredName());
-				if (OutFailReason)
-				{
-					*OutFailReason = FText::Format(LOCTEXT("FailImportTextFromObject", "Unable to import JSON object with invalid keys into Text property {0}"), FText::FromString(Property->GetAuthoredName()));
-				}
-				return false;
-			}
-			TextProperty->SetPropertyValue(OutValue, Text);
-		}
-		else
-		{
-			UE_LOG(LogJson, Error, TEXT("JsonValueToUProperty - Unable to import JSON value that is neither string nor object into Text property %s"), *Property->GetAuthoredName());
-			if (OutFailReason)
-			{
-				*OutFailReason = FText::Format(LOCTEXT("FailImportText", "Unable to import JSON value that is neither string nor object into Text property {0}"), FText::FromString(Property->GetAuthoredName()));
-			}
-			return false;
-		}
-	}
-	else if (FStructProperty *StructProperty = CastField<FStructProperty>(Property))
-	{
-		if (JsonValue->Type == EJson::Object)
-		{
-			TSharedPtr<FJsonObject> Obj = JsonValue->AsObject();
-			check(Obj.IsValid()); // should not fail if Type == EJson::Object
-			if (!JsonAttributesToUStructWithContainer(Obj->Values, StructProperty->Struct, OutValue, ContainerStruct, Container, CheckFlags & (~CPF_ParmFlags), SkipFlags, bStrictMode, OutFailReason))
-			{
-				UE_LOG(LogJson, Error, TEXT("JsonValueToUProperty - Unable to import JSON object into %s property %s"), *StructProperty->Struct->GetAuthoredName(), *Property->GetAuthoredName());
-				if (OutFailReason)
-				{
-					*OutFailReason = FText::Format(LOCTEXT("FailImportStructFromObject", "Unable to import JSON object into {0} property {1}\n{2}"), FText::FromString(StructProperty->Struct->GetAuthoredName()), FText::FromString(Property->GetAuthoredName()), *OutFailReason);
-				}
-				return false;
-			}
-		}
-		else if (JsonValue->Type == EJson::String && StructProperty->Struct->GetFName() == NAME_LinearColor)
-		{
-			FLinearColor& ColorOut = *(FLinearColor*)OutValue;
-			FString ColorString = JsonValue->AsString();
-
-			FColor IntermediateColor;
-			IntermediateColor = FColor::FromHex(ColorString);
-
-			ColorOut = IntermediateColor;
-		}
-		else if (JsonValue->Type == EJson::String && StructProperty->Struct->GetFName() == NAME_Color)
-		{
-			FColor& ColorOut = *(FColor*)OutValue;
-			FString ColorString = JsonValue->AsString();
-
-			ColorOut = FColor::FromHex(ColorString);
-		}
-		else if (JsonValue->Type == EJson::String && StructProperty->Struct->GetFName() == NAME_DateTime)
-		{
-			FString DateString = JsonValue->AsString();
-			FDateTime& DateTimeOut = *(FDateTime*)OutValue;
-			if (DateString == TEXT("min"))
-			{
-				// min representable value for our date struct. Actual date may vary by platform (this is used for sorting)
-				DateTimeOut = FDateTime::MinValue();
-			}
-			else if (DateString == TEXT("max"))
-			{
-				// max representable value for our date struct. Actual date may vary by platform (this is used for sorting)
-				DateTimeOut = FDateTime::MaxValue();
-			}
-			else if (DateString == TEXT("now"))
-			{
-				// this value's not really meaningful from JSON serialization (since we don't know timezone) but handle it anyway since we're handling the other keywords
-				DateTimeOut = FDateTime::UtcNow();
-			}
-			else if (FDateTime::ParseIso8601(*DateString, DateTimeOut))
-			{
-				// ok
-			}
-			else if (FDateTime::Parse(DateString, DateTimeOut))
-			{
-				// ok
-			}
-			else
-			{
-				UE_LOG(LogJson, Error, TEXT("JsonValueToUProperty - Unable to import JSON string into DateTime property %s"), *Property->GetAuthoredName());
-				if (OutFailReason)
-				{
-					*OutFailReason = FText::Format(LOCTEXT("FailImportDateTimeFromString", "Unable to import JSON string into DateTime property {0}"), FText::FromString(Property->GetAuthoredName()));
-				}
-				return false;
-			}
-		}
-		else if (JsonValue->Type == EJson::String && StructProperty->Struct->GetCppStructOps() && StructProperty->Struct->GetCppStructOps()->HasImportTextItem())
-		{
-			UScriptStruct::ICppStructOps* TheCppStructOps = StructProperty->Struct->GetCppStructOps();
-
-			FString ImportTextString = JsonValue->AsString();
-			const TCHAR* ImportTextPtr = *ImportTextString;
-			if (!TheCppStructOps->ImportTextItem(ImportTextPtr, OutValue, PPF_None, nullptr, (FOutputDevice*)GWarn))
-			{
-				// Fall back to trying the tagged property approach if custom ImportTextItem couldn't get it done
-				if (Property->ImportText_Direct(ImportTextPtr, OutValue, nullptr, PPF_None) == nullptr)
-				{
-					UE_LOG(LogJson, Error, TEXT("JsonValueToUProperty - Unable to import JSON string into %s property %s"), *StructProperty->Struct->GetAuthoredName(), *Property->GetAuthoredName());
-					if (OutFailReason)
-					{
-						*OutFailReason = FText::Format(LOCTEXT("FailImportStructFromString", "Unable to import JSON string into {0} property {1}"), FText::FromString(StructProperty->Struct->GetAuthoredName()), FText::FromString(Property->GetAuthoredName()));
-					}
-					return false;
-				}
-			}
-		}
-		else if (JsonValue->Type == EJson::String)
-		{
-			FString ImportTextString = JsonValue->AsString();
-			const TCHAR* ImportTextPtr = *ImportTextString;
-			if (Property->ImportText_Direct(ImportTextPtr, OutValue, nullptr, PPF_None) == nullptr)
-			{
-				UE_LOG(LogJson, Error, TEXT("JsonValueToUProperty - Unable to import JSON string into %s property %s"), *StructProperty->Struct->GetAuthoredName(), *Property->GetAuthoredName());
-				if (OutFailReason)
-				{
-					*OutFailReason = FText::Format(LOCTEXT("FailImportStructFromString", "Unable to import JSON string into {0} property {1}"), FText::FromString(StructProperty->Struct->GetAuthoredName()), FText::FromString(Property->GetAuthoredName()));
-				}
-				return false;
-			}
-		}
-		else
-		{
-			UE_LOG(LogJson, Error, TEXT("JsonValueToUProperty - Unable to import JSON value that is neither string nor object into %s property %s"), *StructProperty->Struct->GetAuthoredName(), *Property->GetAuthoredName());
-			if (OutFailReason)
-			{
-				*OutFailReason = FText::Format(LOCTEXT("FailImportStruct", "Unable to import JSON value that is neither string nor object into {0} property {1}"), FText::FromString(StructProperty->Struct->GetAuthoredName()), FText::FromString(Property->GetAuthoredName()));
-			}
-			return false;
-		}
-	}
-	else if (FObjectProperty *ObjectProperty = CastField<FObjectProperty>(Property))
+	    if (FEnumProperty* EnumProperty = CastField<FEnumProperty>(Property))
+	    {
+		    if (JsonValue->Type == EJson::String)
+		    {
+			    // see if we were passed a string for the enum
+			    const UEnum* Enum = EnumProperty->GetEnum();
+			    check(Enum);
+			    FString StrValue = JsonValue->AsString();
+			    int64 IntValue = Enum->GetValueByName(FName(*StrValue), EGetByNameFlags::CheckAuthoredName);
+			    if (IntValue == INDEX_NONE)
+			    {
+				    UE_LOG(LogJson, Error, TEXT("JsonValueToUProperty - Unable to import enum %s from string value %s for property %s"), *Enum->CppType, *StrValue, *Property->GetAuthoredName());
+				    if (OutFailReason)
+				    {
+					    *OutFailReason = FText::Format(LOCTEXT("FailImportEnumFromString", "Unable to import enum {0} from string value {1} for property {2}"), FText::FromString(Enum->CppType), FText::FromString(StrValue), FText::FromString(Property->GetAuthoredName()));
+				    }
+				    return false;
+			    }
+			    EnumProperty->GetUnderlyingProperty()->SetIntPropertyValue(OutValue, IntValue);
+		    }
+		    else
+		    {
+			    // AsNumber will log an error for completely inappropriate types (then give us a default)
+			    EnumProperty->GetUnderlyingProperty()->SetIntPropertyValue(OutValue, (int64)JsonValue->AsNumber());
+		    }
+	    }
+	    else if (FNumericProperty *NumericProperty = CastField<FNumericProperty>(Property))
+	    {
+		    if (NumericProperty->IsEnum() && JsonValue->Type == EJson::String)
+		    {
+			    // see if we were passed a string for the enum
+			    const UEnum* Enum = NumericProperty->GetIntPropertyEnum();
+			    check(Enum); // should be assured by IsEnum()
+			    FString StrValue = JsonValue->AsString();
+			    int64 IntValue = Enum->GetValueByName(FName(*StrValue), EGetByNameFlags::CheckAuthoredName);
+			    if (IntValue == INDEX_NONE)
+			    {
+				    UE_LOG(LogJson, Error, TEXT("JsonValueToUProperty - Unable to import enum %s from numeric value %s for property %s"), *Enum->CppType, *StrValue, *Property->GetAuthoredName());
+				    if (OutFailReason)
+				    {
+					    *OutFailReason = FText::Format(LOCTEXT("FailImportEnumFromNumeric", "Unable to import enum {0} from numeric value {1} for property {2}"), FText::FromString(Enum->CppType), FText::FromString(StrValue), FText::FromString(Property->GetAuthoredName()));
+				    }
+				    return false;
+			    }
+			    NumericProperty->SetIntPropertyValue(OutValue, IntValue);
+		    }
+		    else if (NumericProperty->IsFloatingPoint())
+		    {
+			    // AsNumber will log an error for completely inappropriate types (then give us a default)
+			    NumericProperty->SetFloatingPointPropertyValue(OutValue, JsonValue->AsNumber());
+		    }
+		    else if (NumericProperty->IsInteger())
+		    {
+			    if (JsonValue->Type == EJson::String)
+			    {
+				    // parse string -> int64 ourselves so we don't lose any precision going through AsNumber (aka double)
+				    NumericProperty->SetIntPropertyValue(OutValue, FCString::Atoi64(*JsonValue->AsString()));
+			    }
+			    else
+			    {
+				    // AsNumber will log an error for completely inappropriate types (then give us a default)
+				    NumericProperty->SetIntPropertyValue(OutValue, (int64)JsonValue->AsNumber());
+			    }
+		    }
+		    else
+		    {
+			    UE_LOG(LogJson, Error, TEXT("JsonValueToUProperty - Unable to import json value into %s numeric property %s"), *Property->GetClass()->GetName(), *Property->GetAuthoredName());
+			    if (OutFailReason)
+			    {
+				    *OutFailReason = FText::Format(LOCTEXT("FailImportNumericProperty", "Unable to import json value into {0} numeric property {1}"), FText::FromString(Property->GetClass()->GetName()), FText::FromString(Property->GetAuthoredName()));
+			    }
+			    return false;
+		    }
+	    }
+	    else if (FBoolProperty *BoolProperty = CastField<FBoolProperty>(Property))
+	    {
+		    // AsBool will log an error for completely inappropriate types (then give us a default)
+		    BoolProperty->SetPropertyValue(OutValue, JsonValue->AsBool());
+	    }
+	    else if (FStrProperty *StringProperty = CastField<FStrProperty>(Property))
+	    {
+		    // AsString will log an error for completely inappropriate types (then give us a default)
+		    StringProperty->SetPropertyValue(OutValue, JsonValue->AsString());
+	    }
+	    else if (FArrayProperty *ArrayProperty = CastField<FArrayProperty>(Property))
+	    {
+		    if (JsonValue->Type == EJson::Array)
+		    {
+			    TArray< TSharedPtr<FJsonValue> > ArrayValue = JsonValue->AsArray();
+			    int32 ArrLen = ArrayValue.Num();
+
+			    // make the output array size match
+			    FScriptArrayHelper Helper(ArrayProperty, OutValue);
+			    Helper.Resize(ArrLen);
+
+			    // set the property values
+			    for (int32 i = 0; i < ArrLen; ++i)
+			    {
+				    const TSharedPtr<FJsonValue>& ArrayValueItem = ArrayValue[i];
+				    if (ArrayValueItem.IsValid() && !ArrayValueItem->IsNull())
+				    {
+					    if (!JsonValueToFPropertyWithContainer(ArrayValueItem, ArrayProperty->Inner, Helper.GetRawPtr(i), ContainerStruct, Container, CheckFlags & (~CPF_ParmFlags), SkipFlags, bStrictMode, OutFailReason, ImportCb))
+					    {
+						    UE_LOG(LogJson, Error, TEXT("JsonValueToUProperty - Unable to import Array element %d for property %s"), i, *Property->GetAuthoredName());
+						    if (OutFailReason)
+						    {
+							    *OutFailReason = FText::Format(LOCTEXT("FailImportArrayElement", "Unable to import Array element {0} for property {1}\n{2}"), FText::AsNumber(i), FText::FromString(Property->GetAuthoredName()), *OutFailReason);
+						    }
+						    return false;
+					    }
+				    }
+			    }
+		    }
+		    else
+		    {
+			    UE_LOG(LogJson, Error, TEXT("JsonValueToUProperty - Unable to import non-array JSON value into Array property %s"), *Property->GetAuthoredName());
+			    if (OutFailReason)
+			    {
+				    *OutFailReason = FText::Format(LOCTEXT("FailImportArray", "Unable to import non-array JSON value into Array property {0}"), FText::FromString(Property->GetAuthoredName()));
+			    }
+			    return false;
+		    }
+	    }
+	    else if (FMapProperty* MapProperty = CastField<FMapProperty>(Property))
+	    {
+		    if (JsonValue->Type == EJson::Object)
+		    {
+				TSharedPtr<FJsonObject> ObjectValue = JsonValue->AsObject();
+
+				FScriptMapHelper Helper(MapProperty, OutValue);
+
+				check(ObjectValue);
+
+			    int32 MapSize = ObjectValue->Values.Num();
+			    Helper.EmptyValues(MapSize);
+
+			    // set the property values
+			    for (const auto& Entry : ObjectValue->Values)
+			    {
+				    if (Entry.Value.IsValid() && !Entry.Value->IsNull())
+				    {
+					    int32 NewIndex = Helper.AddDefaultValue_Invalid_NeedsRehash();
+
+						TSharedPtr<FJsonValueString> TempKeyValue = MakeShared<FJsonValueString>(Entry.Key);
+
+					    if (!JsonValueToFPropertyWithContainer(TempKeyValue, MapProperty->KeyProp, Helper.GetKeyPtr(NewIndex), ContainerStruct, Container, CheckFlags & (~CPF_ParmFlags), SkipFlags, bStrictMode, OutFailReason, ImportCb))
+					    {
+						    UE_LOG(LogJson, Error, TEXT("JsonValueToUProperty - Unable to import Map element %s key for property %s"), *Entry.Key, *Property->GetAuthoredName());
+						    if (OutFailReason)
+						    {
+							    *OutFailReason = FText::Format(LOCTEXT("FailImportMapElementKey", "Unable to import Map element {0} key for property {1}\n{2}"), FText::FromString(Entry.Key), FText::FromString(Property->GetAuthoredName()), *OutFailReason);
+						    }
+						    return false;
+					    }
+
+					    if (!JsonValueToFPropertyWithContainer(Entry.Value, MapProperty->ValueProp, Helper.GetValuePtr(NewIndex), ContainerStruct, Container, CheckFlags & (~CPF_ParmFlags), SkipFlags, bStrictMode, OutFailReason, ImportCb))
+					    {
+						    UE_LOG(LogJson, Error, TEXT("JsonValueToUProperty - Unable to import Map element %s value for property %s"), *Entry.Key, *Property->GetAuthoredName());
+						    if (OutFailReason)
+						    {
+							    *OutFailReason = FText::Format(LOCTEXT("FailImportMapElementValue", "Unable to import Map element {0} value for property {1}\n{2}"), FText::FromString(Entry.Key), FText::FromString(Property->GetAuthoredName()), *OutFailReason);
+						    }
+						    return false;
+					    }
+				    }
+			    }
+
+			    Helper.Rehash();
+		    }
+		    else
+		    {
+			    UE_LOG(LogJson, Error, TEXT("JsonValueToUProperty - Unable to import non-object JSON value into Map property %s"), *Property->GetAuthoredName());
+			    if (OutFailReason)
+			    {
+				    *OutFailReason = FText::Format(LOCTEXT("FailImportMap", "Unable to import non-object JSON value into Map property {0}"), FText::FromString(Property->GetAuthoredName()));
+			    }
+			    return false;
+		    }
+	    }
+	    else if (FSetProperty* SetProperty = CastField<FSetProperty>(Property))
+	    {
+		    if (JsonValue->Type == EJson::Array)
+		    {
+			    TArray< TSharedPtr<FJsonValue> > ArrayValue = JsonValue->AsArray();
+			    int32 ArrLen = ArrayValue.Num();
+
+			    FScriptSetHelper Helper(SetProperty, OutValue);
+			    Helper.EmptyElements(ArrLen);
+
+			    // set the property values
+			    for (int32 i = 0; i < ArrLen; ++i)
+			    {
+				    const TSharedPtr<FJsonValue>& ArrayValueItem = ArrayValue[i];
+				    if (ArrayValueItem.IsValid() && !ArrayValueItem->IsNull())
+				    {
+					    int32 NewIndex = Helper.AddDefaultValue_Invalid_NeedsRehash();
+					    if (!JsonValueToFPropertyWithContainer(ArrayValueItem, SetProperty->ElementProp, Helper.GetElementPtr(NewIndex), ContainerStruct, Container, CheckFlags & (~CPF_ParmFlags), SkipFlags, bStrictMode, OutFailReason, ImportCb))
+					    {
+						    UE_LOG(LogJson, Error, TEXT("JsonValueToUProperty - Unable to import Set element %d for property %s"), i, *Property->GetAuthoredName());
+						    if (OutFailReason)
+						    {
+							    *OutFailReason = FText::Format(LOCTEXT("FailImportSetElement", "Unable to import Set element {0} for property {1}\n{2}"), FText::AsNumber(i), FText::FromString(Property->GetAuthoredName()), *OutFailReason);
+						    }
+						    return false;
+					    }
+				    }
+			    }
+
+			    Helper.Rehash();
+		    }
+		    else
+		    {
+			    UE_LOG(LogJson, Error, TEXT("JsonValueToUProperty - Unable to import non-array JSON value into Set property %s"), *Property->GetAuthoredName());
+			    if (OutFailReason)
+			    {
+				    *OutFailReason = FText::Format(LOCTEXT("FailImportSet", "Unable to import non-array JSON value into Set property {0}"), FText::FromString(Property->GetAuthoredName()));
+			    }
+			    return false;
+		    }
+	    }
+	    else if (FTextProperty* TextProperty = CastField<FTextProperty>(Property))
+	    {
+		    if (JsonValue->Type == EJson::String)
+		    {
+			    // assume this string is already localized, so import as invariant
+			    TextProperty->SetPropertyValue(OutValue, FText::FromString(JsonValue->AsString()));
+		    }
+		    else if (JsonValue->Type == EJson::Object)
+		    {
+			    TSharedPtr<FJsonObject> Obj = JsonValue->AsObject();
+			    check(Obj.IsValid()); // should not fail if Type == EJson::Object
+    
+			    // import the subvalue as a culture invariant string
+			    FText Text;
+			    if (!FJsonObjectConverter::GetTextFromObject(Obj.ToSharedRef(), Text))
+			    {
+				    UE_LOG(LogJson, Error, TEXT("JsonValueToUProperty - Unable to import JSON object with invalid keys into Text property %s"), *Property->GetAuthoredName());
+				    if (OutFailReason)
+				    {
+					    *OutFailReason = FText::Format(LOCTEXT("FailImportTextFromObject", "Unable to import JSON object with invalid keys into Text property {0}"), FText::FromString(Property->GetAuthoredName()));
+				    }
+				    return false;
+			    }
+			    TextProperty->SetPropertyValue(OutValue, Text);
+		    }
+		    else
+		    {
+			    UE_LOG(LogJson, Error, TEXT("JsonValueToUProperty - Unable to import JSON value that is neither string nor object into Text property %s"), *Property->GetAuthoredName());
+			    if (OutFailReason)
+			    {
+				    *OutFailReason = FText::Format(LOCTEXT("FailImportText", "Unable to import JSON value that is neither string nor object into Text property {0}"), FText::FromString(Property->GetAuthoredName()));
+			    }
+			    return false;
+		    }
+	    }
+	    else if (FStructProperty *StructProperty = CastField<FStructProperty>(Property))
+	    {
+		    if (JsonValue->Type == EJson::Object)
+		    {
+			    TSharedPtr<FJsonObject> Obj = JsonValue->AsObject();
+			    check(Obj.IsValid()); // should not fail if Type == EJson::Object
+			    if (!JsonAttributesToUStructWithContainer(Obj->Values, StructProperty->Struct, OutValue, ContainerStruct, Container, CheckFlags & (~CPF_ParmFlags), SkipFlags, bStrictMode, OutFailReason, ImportCb))
+			    {
+				    UE_LOG(LogJson, Error, TEXT("JsonValueToUProperty - Unable to import JSON object into %s property %s"), *StructProperty->Struct->GetAuthoredName(), *Property->GetAuthoredName());
+				    if (OutFailReason)
+				    {
+					    *OutFailReason = FText::Format(LOCTEXT("FailImportStructFromObject", "Unable to import JSON object into {0} property {1}\n{2}"), FText::FromString(StructProperty->Struct->GetAuthoredName()), FText::FromString(Property->GetAuthoredName()), *OutFailReason);
+				    }
+				    return false;
+			    }
+		    }
+		    else if (JsonValue->Type == EJson::String && StructProperty->Struct->GetFName() == NAME_LinearColor)
+		    {
+			    FLinearColor& ColorOut = *(FLinearColor*)OutValue;
+			    FString ColorString = JsonValue->AsString();
+    
+			    FColor IntermediateColor;
+			    IntermediateColor = FColor::FromHex(ColorString);
+    
+			    ColorOut = IntermediateColor;
+		    }
+		    else if (JsonValue->Type == EJson::String && StructProperty->Struct->GetFName() == NAME_Color)
+		    {
+			    FColor& ColorOut = *(FColor*)OutValue;
+			    FString ColorString = JsonValue->AsString();
+
+			    ColorOut = FColor::FromHex(ColorString);
+		    }
+		    else if (JsonValue->Type == EJson::String && StructProperty->Struct->GetFName() == NAME_DateTime)
+		    {
+			    FString DateString = JsonValue->AsString();
+			    FDateTime& DateTimeOut = *(FDateTime*)OutValue;
+			    if (DateString == TEXT("min"))
+			    {
+				    // min representable value for our date struct. Actual date may vary by platform (this is used for sorting)
+				    DateTimeOut = FDateTime::MinValue();
+			    }
+			    else if (DateString == TEXT("max"))
+			    {
+				    // max representable value for our date struct. Actual date may vary by platform (this is used for sorting)
+				    DateTimeOut = FDateTime::MaxValue();
+			    }
+			    else if (DateString == TEXT("now"))
+			    {
+				    // this value's not really meaningful from JSON serialization (since we don't know timezone) but handle it anyway since we're handling the other keywords
+				    DateTimeOut = FDateTime::UtcNow();
+			    }
+			    else if (FDateTime::ParseIso8601(*DateString, DateTimeOut))
+			    {
+				    // ok
+			    }
+			    else if (FDateTime::Parse(DateString, DateTimeOut))
+			    {
+				    // ok
+			    }
+			    else
+			    {
+				    UE_LOG(LogJson, Error, TEXT("JsonValueToUProperty - Unable to import JSON string into DateTime property %s"), *Property->GetAuthoredName());
+				    if (OutFailReason)
+				    {
+					    *OutFailReason = FText::Format(LOCTEXT("FailImportDateTimeFromString", "Unable to import JSON string into DateTime property {0}"), FText::FromString(Property->GetAuthoredName()));
+				    }
+				    return false;
+			    }
+		    }
+		    else if (JsonValue->Type == EJson::String && StructProperty->Struct->GetCppStructOps() && StructProperty->Struct->GetCppStructOps()->HasImportTextItem())
+		    {
+			    UScriptStruct::ICppStructOps* TheCppStructOps = StructProperty->Struct->GetCppStructOps();
+
+			    FString ImportTextString = JsonValue->AsString();
+			    const TCHAR* ImportTextPtr = *ImportTextString;
+			    if (!TheCppStructOps->ImportTextItem(ImportTextPtr, OutValue, PPF_None, nullptr, (FOutputDevice*)GWarn))
+			    {
+				    // Fall back to trying the tagged property approach if custom ImportTextItem couldn't get it done
+				    if (Property->ImportText_Direct(ImportTextPtr, OutValue, nullptr, PPF_None) == nullptr)
+				    {
+					    UE_LOG(LogJson, Error, TEXT("JsonValueToUProperty - Unable to import JSON string into %s property %s"), *StructProperty->Struct->GetAuthoredName(), *Property->GetAuthoredName());
+					    if (OutFailReason)
+					    {
+						    *OutFailReason = FText::Format(LOCTEXT("FailImportStructFromString", "Unable to import JSON string into {0} property {1}"), FText::FromString(StructProperty->Struct->GetAuthoredName()), FText::FromString(Property->GetAuthoredName()));
+					    }
+					    return false;
+				    }
+			    }
+		    }
+		    else if (JsonValue->Type == EJson::String)
+		    {
+			    FString ImportTextString = JsonValue->AsString();
+			    const TCHAR* ImportTextPtr = *ImportTextString;
+			    if (Property->ImportText_Direct(ImportTextPtr, OutValue, nullptr, PPF_None) == nullptr)
+			    {
+				    UE_LOG(LogJson, Error, TEXT("JsonValueToUProperty - Unable to import JSON string into %s property %s"), *StructProperty->Struct->GetAuthoredName(), *Property->GetAuthoredName());
+				    if (OutFailReason)
+				    {
+					    *OutFailReason = FText::Format(LOCTEXT("FailImportStructFromString", "Unable to import JSON string into {0} property {1}"), FText::FromString(StructProperty->Struct->GetAuthoredName()), FText::FromString(Property->GetAuthoredName()));
+				    }
+				    return false;
+			    }
+		    }
+		    else
+		    {
+			    UE_LOG(LogJson, Error, TEXT("JsonValueToUProperty - Unable to import JSON value that is neither string nor object into %s property %s"), *StructProperty->Struct->GetAuthoredName(), *Property->GetAuthoredName());
+			    if (OutFailReason)
+			    {
+				    *OutFailReason = FText::Format(LOCTEXT("FailImportStruct", "Unable to import JSON value that is neither string nor object into {0} property {1}"), FText::FromString(StructProperty->Struct->GetAuthoredName()), FText::FromString(Property->GetAuthoredName()));
+			    }
+			    return false;
+		    }
+	    }
+	    else if (FObjectProperty *ObjectProperty = CastField<FObjectProperty>(Property))
 		{
 			if (JsonValue->Type == EJson::Object)
 			{
@@ -775,7 +775,7 @@ namespace
 				ObjectProperty->SetObjectPropertyValue(OutValue, createdObj);
 
 				check(Obj.IsValid()); // should not fail if Type == EJson::Object
-				if (!JsonAttributesToUStructWithContainer(Obj->Values, PropertyClass, createdObj, PropertyClass, createdObj, CheckFlags & (~CPF_ParmFlags), SkipFlags, bStrictMode, OutFailReason))
+				if (!JsonAttributesToUStructWithContainer(Obj->Values, PropertyClass, createdObj, PropertyClass, createdObj, CheckFlags & (~CPF_ParmFlags), SkipFlags, bStrictMode, OutFailReason, ImportCb))
 				{
 					UE_LOG(LogJson, Error, TEXT("JsonValueToUProperty - Unable to import JSON object into %s property %s"), *PropertyClass->GetAuthoredName(), *Property->GetAuthoredName());
 					if (OutFailReason)
@@ -799,25 +799,25 @@ namespace
 				}
 			}
 		}
-	else
-	{
-		// Default to expect a string for everything else
-		if (Property->ImportText_Direct(*JsonValue->AsString(), OutValue, nullptr, 0) == nullptr)
+		else
 		{
-			UE_LOG(LogJson, Error, TEXT("JsonValueToUProperty - Unable to import JSON string into property %s"), *Property->GetAuthoredName());
-			if (OutFailReason)
-			{
-				*OutFailReason = FText::Format(LOCTEXT("FailImportFromString", "Unable to import JSON string into property {0}"), FText::FromString(Property->GetAuthoredName()));
-			}
-			return false;
-		}
-	}
-
-	return true;
+		    // Default to expect a string for everything else
+		    if (Property->ImportText_Direct(*JsonValue->AsString(), OutValue, nullptr, 0) == nullptr)
+		    {
+			    UE_LOG(LogJson, Error, TEXT("JsonValueToUProperty - Unable to import JSON string into property %s"), *Property->GetAuthoredName());
+			    if (OutFailReason)
+			    {
+				    *OutFailReason = FText::Format(LOCTEXT("FailImportFromString", "Unable to import JSON string into property {0}"), FText::FromString(Property->GetAuthoredName()));
+			    }
+			    return false;
+		    }
+	    }
+
+	    return true;
 	}
 
 
-	bool JsonValueToFPropertyWithContainer(const TSharedPtr<FJsonValue>& JsonValue, FProperty* Property, void* OutValue, const UStruct* ContainerStruct, void* Container, int64 CheckFlags, int64 SkipFlags, const bool bStrictMode, FText* OutFailReason)
+	bool JsonValueToFPropertyWithContainer(const TSharedPtr<FJsonValue>& JsonValue, FProperty* Property, void* OutValue, const UStruct* ContainerStruct, void* Container, int64 CheckFlags, int64 SkipFlags, const bool bStrictMode, FText* OutFailReason, const FJsonObjectConverter::CustomImportCallback* ImportCb)
 	{
 		if (!JsonValue.IsValid())
 		{
@@ -859,14 +859,14 @@ namespace
 				UE_LOG(LogJson, Warning, TEXT("Ignoring excess properties when deserializing %s"), *Property->GetAuthoredName());
 			}
 
-			return ConvertScalarJsonValueToFPropertyWithContainer(JsonValue, Property, OutValue, ContainerStruct, Container, CheckFlags, SkipFlags, bStrictMode, OutFailReason);
+			return ConvertScalarJsonValueToFPropertyWithContainer(JsonValue, Property, OutValue, ContainerStruct, Container, CheckFlags, SkipFlags, bStrictMode, OutFailReason, ImportCb);
 		}
 
 		// In practice, the ArrayDim == 1 check ought to be redundant, since nested arrays of FProperties are not supported
 		if (bArrayOrSetProperty && Property->ArrayDim == 1)
 		{
 			// Read into TArray
-			return ConvertScalarJsonValueToFPropertyWithContainer(JsonValue, Property, OutValue, ContainerStruct, Container, CheckFlags, SkipFlags, bStrictMode, OutFailReason);
+			return ConvertScalarJsonValueToFPropertyWithContainer(JsonValue, Property, OutValue, ContainerStruct, Container, CheckFlags, SkipFlags, bStrictMode, OutFailReason, ImportCb);
 		}
 
 		// We're deserializing a JSON array
@@ -891,7 +891,7 @@ namespace
 		const int32 ItemsToRead = FMath::Clamp(ArrayValue.Num(), 0, Property->ArrayDim);
 		for (int Index = 0; Index != ItemsToRead; ++Index)
 		{
-			if (!ConvertScalarJsonValueToFPropertyWithContainer(ArrayValue[Index], Property, static_cast<char*>(OutValue) + Index * Property->ElementSize, ContainerStruct, Container, CheckFlags, SkipFlags, bStrictMode, OutFailReason))
+			if (!ConvertScalarJsonValueToFPropertyWithContainer(ArrayValue[Index], Property, static_cast<char*>(OutValue) + Index * Property->ElementSize, ContainerStruct, Container, CheckFlags, SkipFlags, bStrictMode, OutFailReason, ImportCb))
 			{
 				return false;
 			}
@@ -899,7 +899,7 @@ namespace
 		return true;
 	}
 
-	bool JsonAttributesToUStructWithContainer(const TMap< FString, TSharedPtr<FJsonValue> >& JsonAttributes, const UStruct* StructDefinition, void* OutStruct, const UStruct* ContainerStruct, void* Container, int64 CheckFlags, int64 SkipFlags, const bool bStrictMode, FText* OutFailReason)
+	bool JsonAttributesToUStructWithContainer(const TMap< FString, TSharedPtr<FJsonValue> >& JsonAttributes, const UStruct* StructDefinition, void* OutStruct, const UStruct* ContainerStruct, void* Container, int64 CheckFlags, int64 SkipFlags, const bool bStrictMode, FText* OutFailReason, const FJsonObjectConverter::CustomImportCallback* ImportCb)
 	{
 		if (StructDefinition == FJsonObjectWrapper::StaticStruct())
 		{
@@ -954,7 +954,13 @@ namespace
 			if (JsonValue->IsValid() && !(*JsonValue)->IsNull())
 			{
 				void* Value = Property->ContainerPtrToValuePtr<uint8>(OutStruct);
-				if (!JsonValueToFPropertyWithContainer(*JsonValue, Property, Value, ContainerStruct, Container, CheckFlags, SkipFlags, bStrictMode, OutFailReason))
+				bool bHasValue = false;
+				if (ImportCb && ImportCb->IsBound())
+				{
+					bHasValue = ImportCb->Execute(*JsonValue, Property, Value);
+				}
+
+				if (!bHasValue && !JsonValueToFPropertyWithContainer(*JsonValue, Property, Value, ContainerStruct, Container, CheckFlags, SkipFlags, bStrictMode, OutFailReason, ImportCb))
 				{
 					UE_LOG(LogJson, Error, TEXT("JsonObjectToUStruct - Unable to import JSON value into property %s"), *PropertyName);
 					if (OutFailReason)
@@ -978,19 +984,19 @@ namespace
 	}
 }
 
-bool FJsonObjectConverter::JsonValueToUProperty(const TSharedPtr<FJsonValue>& JsonValue, FProperty* Property, void* OutValue, int64 CheckFlags, int64 SkipFlags, const bool bStrictMode, FText* OutFailReason)
+bool FJsonObjectConverter::JsonValueToUProperty(const TSharedPtr<FJsonValue>& JsonValue, FProperty* Property, void* OutValue, int64 CheckFlags, int64 SkipFlags, const bool bStrictMode, FText* OutFailReason, const CustomImportCallback* ImportCb)
 {
-	return JsonValueToFPropertyWithContainer(JsonValue, Property, OutValue, nullptr, nullptr, CheckFlags, SkipFlags, bStrictMode, OutFailReason);
+	return JsonValueToFPropertyWithContainer(JsonValue, Property, OutValue, nullptr, nullptr, CheckFlags, SkipFlags, bStrictMode, OutFailReason, ImportCb);
 }
 
-bool FJsonObjectConverter::JsonObjectToUStruct(const TSharedRef<FJsonObject>& JsonObject, const UStruct* StructDefinition, void* OutStruct, int64 CheckFlags, int64 SkipFlags, const bool bStrictMode, FText* OutFailReason)
+bool FJsonObjectConverter::JsonObjectToUStruct(const TSharedRef<FJsonObject>& JsonObject, const UStruct* StructDefinition, void* OutStruct, int64 CheckFlags, int64 SkipFlags, const bool bStrictMode, FText* OutFailReason, const CustomImportCallback* ImportCb)
 {
-	return JsonAttributesToUStruct(JsonObject->Values, StructDefinition, OutStruct, CheckFlags, SkipFlags, bStrictMode, OutFailReason);
+	return JsonAttributesToUStruct(JsonObject->Values, StructDefinition, OutStruct, CheckFlags, SkipFlags, bStrictMode, OutFailReason, ImportCb);
 }
 
-bool FJsonObjectConverter::JsonAttributesToUStruct(const TMap< FString, TSharedPtr<FJsonValue> >& JsonAttributes, const UStruct* StructDefinition, void* OutStruct, int64 CheckFlags, int64 SkipFlags, const bool bStrictMode, FText* OutFailReason)
+bool FJsonObjectConverter::JsonAttributesToUStruct(const TMap< FString, TSharedPtr<FJsonValue> >& JsonAttributes, const UStruct* StructDefinition, void* OutStruct, int64 CheckFlags, int64 SkipFlags, const bool bStrictMode, FText* OutFailReason, const CustomImportCallback* ImportCb)
 {
-	return JsonAttributesToUStructWithContainer(JsonAttributes, StructDefinition, OutStruct, StructDefinition, OutStruct, CheckFlags, SkipFlags, bStrictMode, OutFailReason);
+	return JsonAttributesToUStructWithContainer(JsonAttributes, StructDefinition, OutStruct, StructDefinition, OutStruct, CheckFlags, SkipFlags, bStrictMode, OutFailReason, ImportCb);
 }
 
 //static 
diff --git a/Engine/Source/Runtime/JsonUtilities/Public/JsonObjectConverter.h b/Engine/Source/Runtime/JsonUtilities/Public/JsonObjectConverter.h
--- a/Engine/Source/Runtime/JsonUtilities/Public/JsonObjectConverter.h
+++ b/Engine/Source/Runtime/JsonUtilities/Public/JsonObjectConverter.h
@@ -59,7 +59,14 @@ public: // UStruct -> JSON
 	 * If this returns a valid value it will be inserted into the export chain.
 	 * If this returns nullptr or is not bound, it will try generic type-specific export behavior before falling back to outputting ExportText as a string.
 	 */
-	DECLARE_DELEGATE_RetVal_TwoParams(TSharedPtr<FJsonValue>, CustomExportCallback, FProperty* /* Property */, const void* /* Value */);
+	using CustomExportCallback = TDelegate<TSharedPtr<FJsonValue>(FProperty* Property, const void* Value)>;
+
+	/**
+	 * Optional callback that will be run when importing a single property from Json.
+	 * If this returns true, it should have successfully turned the Json value into the property value.
+	 * If this returns false or is not bound, it will try generic type-specific import behavior before failing.
+	 */
+	using CustomImportCallback = TDelegate<bool(const TSharedPtr<FJsonValue>& JsonValue, FProperty* Property, const void* Value)>;
 
 	static JSONUTILITIES_API const CustomExportCallback ExportCallback_WriteISO8601Dates;
 
@@ -207,10 +214,11 @@ public: // JSON -> UStruct
 	 * @param SkipFlags Skip properties that match any of these flags
 	 * @param bStrictMode Whether to strictly check the json attributes
 	 * @param OutFailReason Reason of the failure if any
+	 * @param ImportCb Optional callback to override import behaviour, if this returns false it will fallback to the default
 	 *
 	 * @return False if any properties matched but failed to deserialize
 	 */
-	static JSONUTILITIES_API bool JsonObjectToUStruct(const TSharedRef<FJsonObject>& JsonObject, const UStruct* StructDefinition, void* OutStruct, int64 CheckFlags = 0, int64 SkipFlags = 0, const bool bStrictMode = false, FText* OutFailReason = nullptr);
+	static JSONUTILITIES_API bool JsonObjectToUStruct(const TSharedRef<FJsonObject>& JsonObject, const UStruct* StructDefinition, void* OutStruct, int64 CheckFlags = 0, int64 SkipFlags = 0, const bool bStrictMode = false, FText* OutFailReason = nullptr, const CustomImportCallback* ImportCb = nullptr);
 
 	/**
 	 * Templated version of JsonObjectToUStruct
@@ -221,19 +229,20 @@ public: // JSON -> UStruct
 	 * @param SkipFlags Skip properties that match any of these flags
 	 * @param bStrictMode Whether to strictly check the json attributes
 	 * @param OutFailReason Reason of the failure if any
+	 * @param ImportCb Optional callback to override import behaviour, if this returns false it will fallback to the default
 	 *
 	 * @return False if any properties matched but failed to deserialize
 	 */
 	template<typename OutStructType>
-	static bool JsonObjectToUStruct(const TSharedRef<FJsonObject>& JsonObject, OutStructType* OutStruct, int64 CheckFlags = 0, int64 SkipFlags = 0, const bool bStrictMode = false, FText* OutFailReason = nullptr)
+	static bool JsonObjectToUStruct(const TSharedRef<FJsonObject>& JsonObject, OutStructType* OutStruct, int64 CheckFlags = 0, int64 SkipFlags = 0, const bool bStrictMode = false, FText* OutFailReason = nullptr, const CustomImportCallback* ImportCb = nullptr)
 	{
 		if constexpr (TModels<CStaticClassProvider, OutStructType>::Value)
 		{
-			return JsonObjectToUStruct(JsonObject, OutStructType::StaticClass(), OutStruct, CheckFlags, SkipFlags, bStrictMode, OutFailReason);
+			return JsonObjectToUStruct(JsonObject, OutStructType::StaticClass(), OutStruct, CheckFlags, SkipFlags, bStrictMode, OutFailReason, ImportCb);
 		}
 		else
 		{
-			return JsonObjectToUStruct(JsonObject, OutStructType::StaticStruct(), OutStruct, CheckFlags, SkipFlags, bStrictMode, OutFailReason);
+			return JsonObjectToUStruct(JsonObject, OutStructType::StaticStruct(), OutStruct, CheckFlags, SkipFlags, bStrictMode, OutFailReason, ImportCb);
 		}
 	}
 
@@ -247,10 +256,11 @@ public: // JSON -> UStruct
 	 * @param SkipFlags Skip properties that match any of these flags
 	 * @param bStrictMode Whether to strictly check the json attributes
 	 * @param OutFailReason Reason of the failure if any
+	 * @param ImportCb Optional callback to override import behaviour, if this returns false it will fallback to the default
 	 *
 	 * @return False if any properties matched but failed to deserialize
 	 */
-	static JSONUTILITIES_API bool JsonAttributesToUStruct(const TMap< FString, TSharedPtr<FJsonValue> >& JsonAttributes, const UStruct* StructDefinition, void* OutStruct, int64 CheckFlags = 0, int64 SkipFlags = 0, const bool bStrictMode = false, FText* OutFailReason = nullptr);
+	static JSONUTILITIES_API bool JsonAttributesToUStruct(const TMap< FString, TSharedPtr<FJsonValue> >& JsonAttributes, const UStruct* StructDefinition, void* OutStruct, int64 CheckFlags = 0, int64 SkipFlags = 0, const bool bStrictMode = false, FText* OutFailReason = nullptr, const CustomImportCallback* ImportCb = nullptr);
 
 	/**
 	 * Converts a single JsonValue to the corresponding FProperty (this may recurse if the property is a UStruct for instance).
@@ -262,10 +272,11 @@ public: // JSON -> UStruct
 	 * @param SkipFlags Skip sub-properties that match any of these flags
 	 * @param bStrictMode Whether to strictly check the json attributes
 	 * @param OutFailReason Reason of the failure if any
+	 * @param ImportCb Optional callback to override import behaviour, if this returns false it will fallback to the default
 	 *
 	 * @return False if the property failed to serialize
 	 */
-	static JSONUTILITIES_API bool JsonValueToUProperty(const TSharedPtr<FJsonValue>& JsonValue, FProperty* Property, void* OutValue, int64 CheckFlags = 0, int64 SkipFlags = 0, const bool bStrictMode = false, FText* OutFailReason = nullptr);
+	static JSONUTILITIES_API bool JsonValueToUProperty(const TSharedPtr<FJsonValue>& JsonValue, FProperty* Property, void* OutValue, int64 CheckFlags = 0, int64 SkipFlags = 0, const bool bStrictMode = false, FText* OutFailReason = nullptr, const CustomImportCallback* ImportCb = nullptr);
 
 	/**
 	 * Converts from a json string containing an object to a UStruct
@@ -276,11 +287,12 @@ public: // JSON -> UStruct
 	 * @param SkipFlags Skip properties that match any of these flags
 	 * @param bStrictMode Whether to strictly check the json attributes
 	 * @param OutFailReason Reason of the failure if any
+	 * @param ImportCb Optional callback to override import behaviour, if this returns false it will fallback to the default
 	 *
 	 * @return False if any properties matched but failed to deserialize
 	 */
 	template<typename OutStructType>
-	static bool JsonObjectStringToUStruct(const FString& JsonString, OutStructType* OutStruct, int64 CheckFlags = 0, int64 SkipFlags = 0, const bool bStrictMode = false, FText* OutFailReason = nullptr)
+	static bool JsonObjectStringToUStruct(const FString& JsonString, OutStructType* OutStruct, int64 CheckFlags = 0, int64 SkipFlags = 0, const bool bStrictMode = false, FText* OutFailReason = nullptr, const CustomImportCallback* ImportCb = nullptr)
 	{
 		TSharedPtr<FJsonObject> JsonObject;
 		TSharedRef<TJsonReader<> > JsonReader = TJsonReaderFactory<>::Create(JsonString);
@@ -293,7 +305,7 @@ public: // JSON -> UStruct
 			}
 			return false;
 		}
-		if (!FJsonObjectConverter::JsonObjectToUStruct(JsonObject.ToSharedRef(), OutStruct, CheckFlags, SkipFlags, bStrictMode, OutFailReason))
+		if (!FJsonObjectConverter::JsonObjectToUStruct(JsonObject.ToSharedRef(), OutStruct, CheckFlags, SkipFlags, bStrictMode, OutFailReason, ImportCb))
 		{
 			UE_LOG(LogJson, Warning, TEXT("JsonObjectStringToUStruct - Unable to deserialize. json=[%s]"), *JsonString);
 			if (OutFailReason)
@@ -314,11 +326,12 @@ public: // JSON -> UStruct
 	 * @param SkipFlags Skip properties that match any of these flags.
 	 * @param bStrictMode Whether to strictly check the json attributes
 	 * @param OutFailReason Reason of the failure if any
+	 * @param ImportCb Optional callback to override import behaviour, if this returns false it will fallback to the default
 	 *
 	 * @return False if any properties matched but failed to deserialize.
 	 */
 	template<typename OutStructType>
-	static bool JsonArrayStringToUStruct(const FString& JsonString, TArray<OutStructType>* OutStructArray, int64 CheckFlags = 0, int64 SkipFlags = 0, const bool bStrictMode = false, FText* OutFailReason = nullptr)
+	static bool JsonArrayStringToUStruct(const FString& JsonString, TArray<OutStructType>* OutStructArray, int64 CheckFlags = 0, int64 SkipFlags = 0, const bool bStrictMode = false, FText* OutFailReason = nullptr, const CustomImportCallback* ImportCb = nullptr)
 	{
 		TArray<TSharedPtr<FJsonValue> > JsonArray;
 		TSharedRef<TJsonReader<> > JsonReader = TJsonReaderFactory<>::Create(JsonString);
@@ -331,7 +344,7 @@ public: // JSON -> UStruct
 			}
 			return false;
 		}
-		if (!JsonArrayToUStruct(JsonArray, OutStructArray, CheckFlags, SkipFlags, bStrictMode, OutFailReason))
+		if (!JsonArrayToUStruct(JsonArray, OutStructArray, CheckFlags, SkipFlags, bStrictMode, OutFailReason, ImportCb))
 		{
 			UE_LOG(LogJson, Warning, TEXT("JsonArrayStringToUStruct - Error parsing one of the elements. json=[%s]"), *JsonString);
 			if (OutFailReason)
@@ -352,11 +365,12 @@ public: // JSON -> UStruct
 	 * @param SkipFlags Skip properties that match any of these flags.
 	 * @param bStrictMode Whether to strictly check the json attributes
 	 * @param OutFailReason Reason of the failure if any
+	 * @param ImportCb Optional callback to override import behaviour, if this returns false it will fallback to the default
 	 *
 	 * @return False if any of the matching elements are not an object, or if one of the matching elements could not be converted to the specified UStruct type.
 	 */
 	template<typename OutStructType>
-	static bool JsonArrayToUStruct(const TArray<TSharedPtr<FJsonValue>>& JsonArray, TArray<OutStructType>* OutStructArray, int64 CheckFlags = 0, int64 SkipFlags = 0, const bool bStrictMode = false, FText* OutFailReason = nullptr)
+	static bool JsonArrayToUStruct(const TArray<TSharedPtr<FJsonValue>>& JsonArray, TArray<OutStructType>* OutStructArray, int64 CheckFlags = 0, int64 SkipFlags = 0, const bool bStrictMode = false, FText* OutFailReason = nullptr, const CustomImportCallback* ImportCb = nullptr)
 	{
 		OutStructArray->SetNum(JsonArray.Num());
 		for (int32 i = 0; i < JsonArray.Num(); ++i)
@@ -371,7 +385,7 @@ public: // JSON -> UStruct
 				}
 				return false;
 			}
-			if (!FJsonObjectConverter::JsonObjectToUStruct(Value->AsObject().ToSharedRef(), OutStructType::StaticStruct(), &(*OutStructArray)[i], CheckFlags, SkipFlags, bStrictMode, OutFailReason))
+			if (!FJsonObjectConverter::JsonObjectToUStruct(Value->AsObject().ToSharedRef(), OutStructType::StaticStruct(), &(*OutStructArray)[i], CheckFlags, SkipFlags, bStrictMode, OutFailReason, ImportCb))
 			{
 				UE_LOG(LogJson, Warning, TEXT("JsonArrayToUStruct - Unable to convert element [%i]."), i);
 				if (OutFailReason)
