Serialize FText as complex string in JsonObjectConverter when the WriteTextAsComplexString flag is used

diff --git a/Engine/Source/Runtime/JsonUtilities/Private/JsonObjectConverter.cpp b/Engine/Source/Runtime/JsonUtilities/Private/JsonObjectConverter.cpp
--- a/Engine/Source/Runtime/JsonUtilities/Private/JsonObjectConverter.cpp
+++ b/Engine/Source/Runtime/JsonUtilities/Private/JsonObjectConverter.cpp
@@ -86,6 +86,14 @@ TSharedPtr<FJsonValue> ConvertScalarFPropertyToJsonValue(FProperty* Property, co
 	}
 	else if (FTextProperty *TextProperty = CastField<FTextProperty>(Property))
 	{
+		if (EnumHasAnyFlags(ConversionFlags, EJsonObjectConversionFlags::WriteTextAsComplexString))
+		{
+			FString TextValueString;
+			FTextStringHelper::WriteToBuffer(TextValueString, TextProperty->GetPropertyValue(Value));
+
+			return MakeShared<FJsonValueString>(TextValueString);
+		}
+
 		return MakeShared<FJsonValueString>(TextProperty->GetPropertyValue(Value).ToString());
 	}
 	else if (FArrayProperty *ArrayProperty = CastField<FArrayProperty>(Property))
@@ -599,8 +607,15 @@ namespace
 		{
 			if (JsonValue->Type == EJson::String)
 			{
+				FString StringValue = JsonValue->AsString();
+				FText TextValue;
+				if (!FTextStringHelper::ReadFromBuffer(*StringValue, TextValue))
+				{
+					TextValue = FText::FromString(StringValue);
+				}
+
 				// assume this string is already localized, so import as invariant
-				TextProperty->SetPropertyValue(OutValue, FText::FromString(JsonValue->AsString()));
+				TextProperty->SetPropertyValue(OutValue, TextValue);
 			}
 			else if (JsonValue->Type == EJson::Object)
 			{
diff --git a/Engine/Source/Runtime/JsonUtilities/Public/JsonObjectConverter.h b/Engine/Source/Runtime/JsonUtilities/Public/JsonObjectConverter.h
--- a/Engine/Source/Runtime/JsonUtilities/Public/JsonObjectConverter.h
+++ b/Engine/Source/Runtime/JsonUtilities/Public/JsonObjectConverter.h
@@ -31,6 +31,12 @@ enum class EJsonObjectConversionFlags
 {
 	None = 0,
 	SkipStandardizeCase = 1 << 0,
+
+	/**
+	 * Write text in its complex exported format (eg, NSLOCTEXT(...)) rather than as a simple string.
+	 * @note This is required to correctly support localization
+	 */
+	WriteTextAsComplexString = 1 << 1
 };
 
 ENUM_CLASS_FLAGS(EJsonObjectConversionFlags)
