Changes to the JsonObjectConverter import callback:
 -Changed the callback to use a non const value pointer
 -Moved the import callback to ConvertScalarJsonValueToFPropertyWithContainer to match where the export callback is called

#tests PIE Added custom import and export callbacks for converting a struct to/from json
#rnx

diff --git a/Engine/Source/Runtime/JsonUtilities/Private/JsonObjectConverter.cpp b/Engine/Source/Runtime/JsonUtilities/Private/JsonObjectConverter.cpp
--- a/Engine/Source/Runtime/JsonUtilities/Private/JsonObjectConverter.cpp
+++ b/Engine/Source/Runtime/JsonUtilities/Private/JsonObjectConverter.cpp
@@ -383,6 +383,15 @@ namespace
 	/** Convert JSON to property, assuming either the property is not an array or the value is an individual array element */
 	bool ConvertScalarJsonValueToFPropertyWithContainer(const TSharedPtr<FJsonValue>& JsonValue, FProperty* Property, void* OutValue, const UStruct* ContainerStruct, void* Container, int64 CheckFlags, int64 SkipFlags, const bool bStrictMode, FText* OutFailReason, const FJsonObjectConverter::CustomImportCallback* ImportCb)
 	{
+		if (ImportCb && ImportCb->IsBound())
+		{
+			if (ImportCb->Execute(JsonValue, Property, OutValue))
+			{
+				return true;
+			}
+			// fall through to default cases
+		}
+
 		if (FEnumProperty* EnumProperty = CastField<FEnumProperty>(Property))
 		{
 			if (JsonValue->Type == EJson::String)
@@ -969,13 +978,7 @@ namespace
 			if (JsonValue->IsValid() && !(*JsonValue)->IsNull())
 			{
 				void* Value = Property->ContainerPtrToValuePtr<uint8>(OutStruct);
-				bool bHasValue = false;
-				if (ImportCb && ImportCb->IsBound())
-				{
-					bHasValue = ImportCb->Execute(*JsonValue, Property, Value);
-				}
-
-				if (!bHasValue && !JsonValueToFPropertyWithContainer(*JsonValue, Property, Value, ContainerStruct, Container, CheckFlags, SkipFlags, bStrictMode, OutFailReason, ImportCb))
+				if (!JsonValueToFPropertyWithContainer(*JsonValue, Property, Value, ContainerStruct, Container, CheckFlags, SkipFlags, bStrictMode, OutFailReason, ImportCb))
 				{
 					UE_LOG(LogJson, Error, TEXT("JsonObjectToUStruct - Unable to import JSON value into property %s"), *PropertyName);
 					if (OutFailReason)
diff --git a/Engine/Source/Runtime/JsonUtilities/Public/JsonObjectConverter.h b/Engine/Source/Runtime/JsonUtilities/Public/JsonObjectConverter.h
--- a/Engine/Source/Runtime/JsonUtilities/Public/JsonObjectConverter.h
+++ b/Engine/Source/Runtime/JsonUtilities/Public/JsonObjectConverter.h
@@ -72,7 +72,7 @@ public: // UStruct -> JSON
 	 * If this returns true, it should have successfully turned the Json value into the property value.
 	 * If this returns false or is not bound, it will try generic type-specific import behavior before failing.
 	 */
-	using CustomImportCallback = TDelegate<bool(const TSharedPtr<FJsonValue>& JsonValue, FProperty* Property, const void* Value)>;
+	using CustomImportCallback = TDelegate<bool(const TSharedPtr<FJsonValue>& JsonValue, FProperty* Property, void* Value)>;
 
 	static JSONUTILITIES_API const CustomExportCallback ExportCallback_WriteISO8601Dates;
 
