Fixed JsonObjectConverter not loading BP classes when deserializing JSON.

diff --git a/Engine/Source/Runtime/JsonUtilities/Private/JsonObjectConverter.cpp b/Engine/Source/Runtime/JsonUtilities/Private/JsonObjectConverter.cpp
--- a/Engine/Source/Runtime/JsonUtilities/Private/JsonObjectConverter.cpp
+++ b/Engine/Source/Runtime/JsonUtilities/Private/JsonObjectConverter.cpp
@@ -786,7 +786,7 @@ namespace
 				Obj->RemoveField(ObjectClassNameKey);
 				if (!ClassString.IsEmpty())
 				{
-					UClass* FoundClass = FPackageName::IsShortPackageName(ClassString) ? FindFirstObject<UClass>(*ClassString) : UClass::TryFindTypeSlow<UClass>(ClassString);
+					UClass* FoundClass = FPackageName::IsShortPackageName(ClassString) ? FindFirstObject<UClass>(*ClassString) : LoadClass<UObject>(nullptr, *ClassString);
 					if (FoundClass)
 					{
 						PropertyClass = FoundClass;
