Forwarded missing flags to inner UPropertyToJsonValue exporting functions.

diff --git a/Engine/Source/Runtime/JsonUtilities/Private/JsonObjectConverter.cpp b/Engine/Source/Runtime/JsonUtilities/Private/JsonObjectConverter.cpp
--- a/Engine/Source/Runtime/JsonUtilities/Private/JsonObjectConverter.cpp
+++ b/Engine/Source/Runtime/JsonUtilities/Private/JsonObjectConverter.cpp
@@ -102,7 +102,7 @@ TSharedPtr<FJsonValue> ConvertScalarFPropertyToJsonValue(FProperty* Property, co
 		FScriptArrayHelper Helper(ArrayProperty, Value);
 		for (int32 i=0, n=Helper.Num(); i<n; ++i)
 		{
-			TSharedPtr<FJsonValue> Elem = FJsonObjectConverter::UPropertyToJsonValue(ArrayProperty->Inner, Helper.GetRawPtr(i), CheckFlags & ( ~CPF_ParmFlags ), SkipFlags, ExportCb, ArrayProperty);
+			TSharedPtr<FJsonValue> Elem = FJsonObjectConverter::UPropertyToJsonValue(ArrayProperty->Inner, Helper.GetRawPtr(i), CheckFlags & ( ~CPF_ParmFlags ), SkipFlags, ExportCb, ArrayProperty, ConversionFlags);
 			if ( Elem.IsValid() )
 			{
 				// add to the array
@@ -117,7 +117,7 @@ TSharedPtr<FJsonValue> ConvertScalarFPropertyToJsonValue(FProperty* Property, co
 		FScriptSetHelper Helper(SetProperty, Value);
 		for (FScriptSetHelper::FIterator It(Helper); It; ++It)
 		{
-			TSharedPtr<FJsonValue> Elem = FJsonObjectConverter::UPropertyToJsonValue(SetProperty->ElementProp, Helper.GetElementPtr(It), CheckFlags & (~CPF_ParmFlags), SkipFlags, ExportCb, SetProperty);
+			TSharedPtr<FJsonValue> Elem = FJsonObjectConverter::UPropertyToJsonValue(SetProperty->ElementProp, Helper.GetElementPtr(It), CheckFlags & (~CPF_ParmFlags), SkipFlags, ExportCb, SetProperty, ConversionFlags);
 			if (Elem.IsValid())
 			{
 				// add to the array
@@ -189,7 +189,7 @@ TSharedPtr<FJsonValue> ConvertScalarFPropertyToJsonValue(FProperty* Property, co
 			TSharedRef<FJsonObject> Out = MakeShared<FJsonObject>();
 
 			Out->SetStringField(ObjectClassNameKey, Object->GetClass()->GetPathName());
-			if (FJsonObjectConverter::UStructToJsonObject(ObjectProperty->GetObjectPropertyValue(Value)->GetClass(), Object, Out, CheckFlags, SkipFlags, ExportCb))
+			if (FJsonObjectConverter::UStructToJsonObject(ObjectProperty->GetObjectPropertyValue(Value)->GetClass(), Object, Out, CheckFlags, SkipFlags, ExportCb, ConversionFlags))
 			{
 				TSharedRef<FJsonValueObject> JsonObject = MakeShared<FJsonValueObject>(Out);
 				JsonObject->Type = EJson::Object;
