[UE5]: Add outer UStruct container params to FProperty->JsonValue conversion API internals.

Notes:
- JsonValue->FProperty internals already included container params (e.g. JsonValueToFPropertyWithContainer), so this adds consistency with the inverse impl.
- I will make use of the container param in a followup change.

diff --git a/Engine/Source/Runtime/JsonUtilities/Private/JsonObjectConverter.cpp b/Engine/Source/Runtime/JsonUtilities/Private/JsonObjectConverter.cpp
--- a/Engine/Source/Runtime/JsonUtilities/Private/JsonObjectConverter.cpp
+++ b/Engine/Source/Runtime/JsonUtilities/Private/JsonObjectConverter.cpp
@@ -31,8 +31,11 @@ namespace
 
 	const FName NAME_DateTime(TEXT("DateTime"));
 
+	TSharedPtr<FJsonValue> FPropertyToJsonValueWithContainer(FProperty* Property, const void* Value, const UObject* Container, int64 CheckFlags, int64 SkipFlags, const FJsonObjectConverter::CustomExportCallback* ExportCb, FProperty* OuterProperty, EJsonObjectConversionFlags ConversionFlags);
+	bool UStructToJsonAttributesWithContainer(const UStruct* StructDefinition, const void* Struct, const UObject* Container, TMap< FString, TSharedPtr<FJsonValue> >& OutJsonAttributes, int64 CheckFlags, int64 SkipFlags, const FJsonObjectConverter::CustomExportCallback* ExportCb, EJsonObjectConversionFlags ConversionFlags);
+
 /** Convert property to JSON, assuming either the property is not an array or the value is an individual array element */
-TSharedPtr<FJsonValue> ConvertScalarFPropertyToJsonValue(FProperty* Property, const void* Value, int64 CheckFlags, int64 SkipFlags, const FJsonObjectConverter::CustomExportCallback* ExportCb, FProperty* OuterProperty, EJsonObjectConversionFlags ConversionFlags)
+TSharedPtr<FJsonValue> ConvertScalarFPropertyToJsonValueWithContainer(FProperty* Property, const void* Value, const UObject* Container, int64 CheckFlags, int64 SkipFlags, const FJsonObjectConverter::CustomExportCallback* ExportCb, FProperty* OuterProperty, EJsonObjectConversionFlags ConversionFlags)
 {
 	// See if there's a custom export callback first, so it can override default behavior
 	if (ExportCb && ExportCb->IsBound())
@@ -102,7 +105,7 @@ TSharedPtr<FJsonValue> ConvertScalarFPropertyToJsonValue(FProperty* Property, co
 		FScriptArrayHelper Helper(ArrayProperty, Value);
 		for (int32 i=0, n=Helper.Num(); i<n; ++i)
 		{
-			TSharedPtr<FJsonValue> Elem = FJsonObjectConverter::UPropertyToJsonValue(ArrayProperty->Inner, Helper.GetRawPtr(i), CheckFlags & ( ~CPF_ParmFlags ), SkipFlags, ExportCb, ArrayProperty, ConversionFlags);
+			TSharedPtr<FJsonValue> Elem = FPropertyToJsonValueWithContainer(ArrayProperty->Inner, Helper.GetRawPtr(i), Container, CheckFlags & ( ~CPF_ParmFlags ), SkipFlags, ExportCb, ArrayProperty, ConversionFlags);
 			if ( Elem.IsValid() )
 			{
 				// add to the array
@@ -117,7 +120,7 @@ TSharedPtr<FJsonValue> ConvertScalarFPropertyToJsonValue(FProperty* Property, co
 		FScriptSetHelper Helper(SetProperty, Value);
 		for (FScriptSetHelper::FIterator It(Helper); It; ++It)
 		{
-			TSharedPtr<FJsonValue> Elem = FJsonObjectConverter::UPropertyToJsonValue(SetProperty->ElementProp, Helper.GetElementPtr(It), CheckFlags & (~CPF_ParmFlags), SkipFlags, ExportCb, SetProperty, ConversionFlags);
+			TSharedPtr<FJsonValue> Elem = FPropertyToJsonValueWithContainer(SetProperty->ElementProp, Helper.GetElementPtr(It), Container, CheckFlags & (~CPF_ParmFlags), SkipFlags, ExportCb, SetProperty, ConversionFlags);
 			if (Elem.IsValid())
 			{
 				// add to the array
@@ -133,8 +136,8 @@ TSharedPtr<FJsonValue> ConvertScalarFPropertyToJsonValue(FProperty* Property, co
 		FScriptMapHelper Helper(MapProperty, Value);
 		for (FScriptMapHelper::FIterator It(Helper); It; ++It)
 		{
-			TSharedPtr<FJsonValue> KeyElement = FJsonObjectConverter::UPropertyToJsonValue(MapProperty->KeyProp, Helper.GetKeyPtr(It), CheckFlags & (~CPF_ParmFlags), SkipFlags, ExportCb, MapProperty, ConversionFlags);
-			TSharedPtr<FJsonValue> ValueElement = FJsonObjectConverter::UPropertyToJsonValue(MapProperty->ValueProp, Helper.GetValuePtr(It), CheckFlags & (~CPF_ParmFlags), SkipFlags, ExportCb, MapProperty, ConversionFlags);
+			TSharedPtr<FJsonValue> KeyElement = FPropertyToJsonValueWithContainer(MapProperty->KeyProp, Helper.GetKeyPtr(It), Container, CheckFlags & (~CPF_ParmFlags), SkipFlags, ExportCb, MapProperty, ConversionFlags);
+			TSharedPtr<FJsonValue> ValueElement = FPropertyToJsonValueWithContainer(MapProperty->ValueProp, Helper.GetValuePtr(It), Container, CheckFlags & (~CPF_ParmFlags), SkipFlags, ExportCb, MapProperty, ConversionFlags);
 			if (KeyElement.IsValid() && ValueElement.IsValid())
 			{
 				FString KeyString;
@@ -175,7 +178,7 @@ TSharedPtr<FJsonValue> ConvertScalarFPropertyToJsonValue(FProperty* Property, co
 		}
 
 		TSharedRef<FJsonObject> Out = MakeShared<FJsonObject>();
-		if (FJsonObjectConverter::UStructToJsonObject(StructProperty->Struct, Value, Out, CheckFlags & (~CPF_ParmFlags), SkipFlags, ExportCb, ConversionFlags))
+		if (UStructToJsonAttributesWithContainer(StructProperty->Struct, Value, Container, Out->Values, CheckFlags & (~CPF_ParmFlags), SkipFlags, ExportCb, ConversionFlags))
 		{
 			return MakeShared<FJsonValueObject>(Out);
 		}
@@ -192,7 +195,8 @@ TSharedPtr<FJsonValue> ConvertScalarFPropertyToJsonValue(FProperty* Property, co
 				Out->SetStringField(ObjectClassNameKey, Object->GetClass()->GetPathName());
 			}
 
-			if (FJsonObjectConverter::UStructToJsonObject(ObjectProperty->GetObjectPropertyValue(Value)->GetClass(), Object, Out, CheckFlags, SkipFlags, ExportCb, ConversionFlags))
+			// Use the subobject as the container context for this conversion
+			if (UStructToJsonAttributesWithContainer(Object->GetClass(), Object, Object, Out->Values, CheckFlags, SkipFlags, ExportCb, ConversionFlags))
 			{
 				TSharedRef<FJsonValueObject> JsonObject = MakeShared<FJsonValueObject>(Out);
 				JsonObject->Type = EJson::Object;
@@ -217,29 +221,23 @@ TSharedPtr<FJsonValue> ConvertScalarFPropertyToJsonValue(FProperty* Property, co
 	// invalid
 	return TSharedPtr<FJsonValue>();
 }
-}
 
-TSharedPtr<FJsonValue> FJsonObjectConverter::UPropertyToJsonValue(FProperty* Property, const void* Value, int64 CheckFlags, int64 SkipFlags, const CustomExportCallback* ExportCb, FProperty* OuterProperty, EJsonObjectConversionFlags ConversionFlags)
+TSharedPtr<FJsonValue> FPropertyToJsonValueWithContainer(FProperty* Property, const void* Value, const UObject* Container, int64 CheckFlags, int64 SkipFlags, const FJsonObjectConverter::CustomExportCallback* ExportCb, FProperty* OuterProperty, EJsonObjectConversionFlags ConversionFlags)
 {
 	if (Property->ArrayDim == 1)
 	{
-		return ConvertScalarFPropertyToJsonValue(Property, Value, CheckFlags, SkipFlags, ExportCb, OuterProperty, ConversionFlags);
+		return ConvertScalarFPropertyToJsonValueWithContainer(Property, Value, Container, CheckFlags, SkipFlags, ExportCb, OuterProperty, ConversionFlags);
 	}
 
 	TArray< TSharedPtr<FJsonValue> > Array;
 	for (int Index = 0; Index != Property->ArrayDim; ++Index)
 	{
-		Array.Add(ConvertScalarFPropertyToJsonValue(Property, (char*)Value + Index * Property->GetElementSize(), CheckFlags, SkipFlags, ExportCb, OuterProperty, ConversionFlags));
+		Array.Add(ConvertScalarFPropertyToJsonValueWithContainer(Property, (char*)Value + Index * Property->GetElementSize(), Container, CheckFlags, SkipFlags, ExportCb, OuterProperty, ConversionFlags));
 	}
 	return MakeShared<FJsonValueArray>(Array);
 }
 
-bool FJsonObjectConverter::UStructToJsonObject(const UStruct* StructDefinition, const void* Struct, TSharedRef<FJsonObject> OutJsonObject, int64 CheckFlags, int64 SkipFlags, const CustomExportCallback* ExportCb, EJsonObjectConversionFlags ConversionFlags)
-{
-	return UStructToJsonAttributes(StructDefinition, Struct, OutJsonObject->Values, CheckFlags, SkipFlags, ExportCb, ConversionFlags);
-}
-
-bool FJsonObjectConverter::UStructToJsonAttributes(const UStruct* StructDefinition, const void* Struct, TMap< FString, TSharedPtr<FJsonValue> >& OutJsonAttributes, int64 CheckFlags, int64 SkipFlags, const CustomExportCallback* ExportCb, EJsonObjectConversionFlags ConversionFlags)
+bool UStructToJsonAttributesWithContainer(const UStruct* StructDefinition, const void* Struct, const UObject* Container, TMap< FString, TSharedPtr<FJsonValue> >& OutJsonAttributes, int64 CheckFlags, int64 SkipFlags, const FJsonObjectConverter::CustomExportCallback* ExportCb, EJsonObjectConversionFlags ConversionFlags)
 {
 	if (SkipFlags == 0)
 	{
@@ -250,7 +248,7 @@ bool FJsonObjectConverter::UStructToJsonAttributes(const UStruct* StructDefiniti
 	if (StructDefinition == FJsonObjectWrapper::StaticStruct())
 	{
 		// Just copy it into the object
-		const FJsonObjectWrapper* ProxyObject = (const FJsonObjectWrapper *)Struct;
+		const FJsonObjectWrapper* ProxyObject = (const FJsonObjectWrapper*)Struct;
 
 		if (ProxyObject->JsonObject.IsValid())
 		{
@@ -276,13 +274,13 @@ bool FJsonObjectConverter::UStructToJsonAttributes(const UStruct* StructDefiniti
 		FString VariableName = Property->GetAuthoredName();
 		if (!EnumHasAnyFlags(ConversionFlags, EJsonObjectConversionFlags::SkipStandardizeCase))
 		{
-			VariableName = StandardizeCase(VariableName);
+			VariableName = FJsonObjectConverter::StandardizeCase(VariableName);
 		}
 
 		const void* Value = Property->ContainerPtrToValuePtr<uint8>(Struct);
 
 		// convert the property to a FJsonValue
-		TSharedPtr<FJsonValue> JsonValue = UPropertyToJsonValue(Property, Value, CheckFlags, SkipFlags, ExportCb, nullptr, ConversionFlags);
+		TSharedPtr<FJsonValue> JsonValue = FPropertyToJsonValueWithContainer(Property, Value, Container, CheckFlags, SkipFlags, ExportCb, nullptr, ConversionFlags);
 		if (!JsonValue.IsValid())
 		{
 			FFieldClass* PropClass = Property->GetClass();
@@ -296,6 +294,28 @@ bool FJsonObjectConverter::UStructToJsonAttributes(const UStruct* StructDefiniti
 
 	return true;
 }
+}
+
+TSharedPtr<FJsonValue> FJsonObjectConverter::UPropertyToJsonValue(FProperty* Property, const void* Value, int64 CheckFlags, int64 SkipFlags, const CustomExportCallback* ExportCb, FProperty* OuterProperty, EJsonObjectConversionFlags ConversionFlags)
+{
+	return FPropertyToJsonValueWithContainer(Property, Value, nullptr, CheckFlags, SkipFlags, ExportCb, OuterProperty, ConversionFlags);
+}
+
+bool FJsonObjectConverter::UStructToJsonObject(const UStruct* StructDefinition, const void* Struct, TSharedRef<FJsonObject> OutJsonObject, int64 CheckFlags, int64 SkipFlags, const CustomExportCallback* ExportCb, EJsonObjectConversionFlags ConversionFlags)
+{
+	return UStructToJsonAttributes(StructDefinition, Struct, OutJsonObject->Values, CheckFlags, SkipFlags, ExportCb, ConversionFlags);
+}
+
+bool FJsonObjectConverter::UStructToJsonAttributes(const UStruct* StructDefinition, const void* Struct, TMap< FString, TSharedPtr<FJsonValue> >& OutJsonAttributes, int64 CheckFlags, int64 SkipFlags, const CustomExportCallback* ExportCb, EJsonObjectConversionFlags ConversionFlags)
+{
+	const UObject* ContainerObject = nullptr;
+	if (StructDefinition->IsA<UClass>())
+	{
+		ContainerObject = static_cast<const UObject*>(Struct);
+	}
+
+	return UStructToJsonAttributesWithContainer(StructDefinition, Struct, ContainerObject, OutJsonAttributes, CheckFlags, SkipFlags, ExportCb, ConversionFlags);
+}
 
 template<class CharType, class PrintPolicy>
 bool UStructToJsonObjectStringInternal(const TSharedRef<FJsonObject>& JsonObject, FString& OutJsonString, int32 Indent)
