[UE5]: Decouple FProperty->JsonValue conversion from instanced property flags to support exporting object properties to nested JSON object values on types w/ dynamic reference values.

Notes:
- Followup to CL 40186645.
- Further extends the internal impl to infer from the previously-added container context when to export object properties by value as an alternative to relying exclusively on its property flags.

diff --git a/Engine/Source/Runtime/JsonUtilities/Private/JsonObjectConverter.cpp b/Engine/Source/Runtime/JsonUtilities/Private/JsonObjectConverter.cpp
--- a/Engine/Source/Runtime/JsonUtilities/Private/JsonObjectConverter.cpp
+++ b/Engine/Source/Runtime/JsonUtilities/Private/JsonObjectConverter.cpp
@@ -31,11 +31,49 @@ namespace
 
 	const FName NAME_DateTime(TEXT("DateTime"));
 
-	TSharedPtr<FJsonValue> FPropertyToJsonValueWithContainer(FProperty* Property, const void* Value, const UObject* Container, int64 CheckFlags, int64 SkipFlags, const FJsonObjectConverter::CustomExportCallback* ExportCb, FProperty* OuterProperty, EJsonObjectConversionFlags ConversionFlags);
-	bool UStructToJsonAttributesWithContainer(const UStruct* StructDefinition, const void* Struct, const UObject* Container, TMap< FString, TSharedPtr<FJsonValue> >& OutJsonAttributes, int64 CheckFlags, int64 SkipFlags, const FJsonObjectConverter::CustomExportCallback* ExportCb, EJsonObjectConversionFlags ConversionFlags);
+	TSharedPtr<FJsonValue> FPropertyToJsonValueWithContainer(FProperty* Property, const void* Value, const UObject* Container, TSet<const UObject*>* ExportedObjects, int64 CheckFlags, int64 SkipFlags, const FJsonObjectConverter::CustomExportCallback* ExportCb, FProperty* OuterProperty, EJsonObjectConversionFlags ConversionFlags);
+	bool UStructToJsonAttributesWithContainer(const UStruct* StructDefinition, const void* Struct, const UObject* Container, TSet<const UObject*>* ExportedObjects, TMap< FString, TSharedPtr<FJsonValue> >& OutJsonAttributes, int64 CheckFlags, int64 SkipFlags, const FJsonObjectConverter::CustomExportCallback* ExportCb, EJsonObjectConversionFlags ConversionFlags);
+	
+	bool ShouldExportObjectPropertyByValue(const FObjectProperty* Property, const UObject* Value, const UObject* Container, const TSet<const UObject*>* ExportedObjects, const FProperty* OuterProperty)
+	{
+		// Nothing to export if NULL.
+		if (!Value)
+		{
+			return false;
+		}
+
+		// Check the instanced flag for backwards compatibility - always export by value in this case.
+		if (Property->HasAnyPropertyFlags(CPF_PersistentInstance)
+			|| (OuterProperty && OuterProperty->HasAnyPropertyFlags(CPF_PersistentInstance)))
+		{
+			return true;
+		}
+
+		// Check if it's a reference to the container (self) to guard against cycles.
+		if (Value == Container)
+		{
+			return false;
+		}
+
+		// Check if we've already exported this value.
+		if (ExportedObjects
+			&& ExportedObjects->Contains(Value))
+		{
+			return false;
+		}
+
+		// Export by value if it is scoped within the current container context (if set).
+		if (Container
+			&& Value->IsInOuter(Container))
+		{
+			return true;
+		}
+
+		return false;
+	}
 
 /** Convert property to JSON, assuming either the property is not an array or the value is an individual array element */
-TSharedPtr<FJsonValue> ConvertScalarFPropertyToJsonValueWithContainer(FProperty* Property, const void* Value, const UObject* Container, int64 CheckFlags, int64 SkipFlags, const FJsonObjectConverter::CustomExportCallback* ExportCb, FProperty* OuterProperty, EJsonObjectConversionFlags ConversionFlags)
+TSharedPtr<FJsonValue> ConvertScalarFPropertyToJsonValueWithContainer(FProperty* Property, const void* Value, const UObject* Container, TSet<const UObject*>* ExportedObjects, int64 CheckFlags, int64 SkipFlags, const FJsonObjectConverter::CustomExportCallback* ExportCb, FProperty* OuterProperty, EJsonObjectConversionFlags ConversionFlags)
 {
 	// See if there's a custom export callback first, so it can override default behavior
 	if (ExportCb && ExportCb->IsBound())
@@ -105,7 +143,7 @@ TSharedPtr<FJsonValue> ConvertScalarFPropertyToJsonValueWithContainer(FProperty*
 		FScriptArrayHelper Helper(ArrayProperty, Value);
 		for (int32 i=0, n=Helper.Num(); i<n; ++i)
 		{
-			TSharedPtr<FJsonValue> Elem = FPropertyToJsonValueWithContainer(ArrayProperty->Inner, Helper.GetRawPtr(i), Container, CheckFlags & ( ~CPF_ParmFlags ), SkipFlags, ExportCb, ArrayProperty, ConversionFlags);
+			TSharedPtr<FJsonValue> Elem = FPropertyToJsonValueWithContainer(ArrayProperty->Inner, Helper.GetRawPtr(i), Container, ExportedObjects, CheckFlags & ( ~CPF_ParmFlags ), SkipFlags, ExportCb, ArrayProperty, ConversionFlags);
 			if ( Elem.IsValid() )
 			{
 				// add to the array
@@ -120,7 +158,7 @@ TSharedPtr<FJsonValue> ConvertScalarFPropertyToJsonValueWithContainer(FProperty*
 		FScriptSetHelper Helper(SetProperty, Value);
 		for (FScriptSetHelper::FIterator It(Helper); It; ++It)
 		{
-			TSharedPtr<FJsonValue> Elem = FPropertyToJsonValueWithContainer(SetProperty->ElementProp, Helper.GetElementPtr(It), Container, CheckFlags & (~CPF_ParmFlags), SkipFlags, ExportCb, SetProperty, ConversionFlags);
+			TSharedPtr<FJsonValue> Elem = FPropertyToJsonValueWithContainer(SetProperty->ElementProp, Helper.GetElementPtr(It), Container, ExportedObjects, CheckFlags & (~CPF_ParmFlags), SkipFlags, ExportCb, SetProperty, ConversionFlags);
 			if (Elem.IsValid())
 			{
 				// add to the array
@@ -136,8 +174,8 @@ TSharedPtr<FJsonValue> ConvertScalarFPropertyToJsonValueWithContainer(FProperty*
 		FScriptMapHelper Helper(MapProperty, Value);
 		for (FScriptMapHelper::FIterator It(Helper); It; ++It)
 		{
-			TSharedPtr<FJsonValue> KeyElement = FPropertyToJsonValueWithContainer(MapProperty->KeyProp, Helper.GetKeyPtr(It), Container, CheckFlags & (~CPF_ParmFlags), SkipFlags, ExportCb, MapProperty, ConversionFlags);
-			TSharedPtr<FJsonValue> ValueElement = FPropertyToJsonValueWithContainer(MapProperty->ValueProp, Helper.GetValuePtr(It), Container, CheckFlags & (~CPF_ParmFlags), SkipFlags, ExportCb, MapProperty, ConversionFlags);
+			TSharedPtr<FJsonValue> KeyElement = FPropertyToJsonValueWithContainer(MapProperty->KeyProp, Helper.GetKeyPtr(It), Container, ExportedObjects, CheckFlags & (~CPF_ParmFlags), SkipFlags, ExportCb, MapProperty, ConversionFlags);
+			TSharedPtr<FJsonValue> ValueElement = FPropertyToJsonValueWithContainer(MapProperty->ValueProp, Helper.GetValuePtr(It), Container, ExportedObjects, CheckFlags & (~CPF_ParmFlags), SkipFlags, ExportCb, MapProperty, ConversionFlags);
 			if (KeyElement.IsValid() && ValueElement.IsValid())
 			{
 				FString KeyString;
@@ -178,7 +216,7 @@ TSharedPtr<FJsonValue> ConvertScalarFPropertyToJsonValueWithContainer(FProperty*
 		}
 
 		TSharedRef<FJsonObject> Out = MakeShared<FJsonObject>();
-		if (UStructToJsonAttributesWithContainer(StructProperty->Struct, Value, Container, Out->Values, CheckFlags & (~CPF_ParmFlags), SkipFlags, ExportCb, ConversionFlags))
+		if (UStructToJsonAttributesWithContainer(StructProperty->Struct, Value, Container, ExportedObjects, Out->Values, CheckFlags & (~CPF_ParmFlags), SkipFlags, ExportCb, ConversionFlags))
 		{
 			return MakeShared<FJsonValueObject>(Out);
 		}
@@ -187,7 +225,7 @@ TSharedPtr<FJsonValue> ConvertScalarFPropertyToJsonValueWithContainer(FProperty*
 	{
 		// Instanced properties should be copied by value, while normal UObject* properties should output as asset references
 		UObject* Object = ObjectProperty->GetObjectPropertyValue(Value);
-		if (Object && (ObjectProperty->HasAnyPropertyFlags(CPF_PersistentInstance) || (OuterProperty && OuterProperty->HasAnyPropertyFlags(CPF_PersistentInstance))))
+		if (ShouldExportObjectPropertyByValue(ObjectProperty, Object, Container, ExportedObjects, OuterProperty))
 		{
 			TSharedRef<FJsonObject> Out = MakeShared<FJsonObject>();
 			if (!EnumHasAnyFlags(ConversionFlags, EJsonObjectConversionFlags::SuppressClassNameForPersistentObject))
@@ -195,8 +233,15 @@ TSharedPtr<FJsonValue> ConvertScalarFPropertyToJsonValueWithContainer(FProperty*
 				Out->SetStringField(ObjectClassNameKey, Object->GetClass()->GetPathName());
 			}
 
-			// Use the subobject as the container context for this conversion
-			if (UStructToJsonAttributesWithContainer(Object->GetClass(), Object, Object, Out->Values, CheckFlags, SkipFlags, ExportCb, ConversionFlags))
+			// Track it to ensure that we only export this object by value once; other instances of this value should export as the object's path (i.e. by reference)
+			if (ExportedObjects)
+			{
+				ExportedObjects->Emplace(Object);
+			}
+
+			// Use the subobject as the container context for this conversion so that we only create inner JsonObject values for instanced subobjects contained within.
+			// Also note we don't clear the ExportedObjects set here to ensure the subobject does not convert references we've already exported by value on an ancestor.
+			if (UStructToJsonAttributesWithContainer(Object->GetClass(), Object, Object, ExportedObjects, Out->Values, CheckFlags, SkipFlags, ExportCb, ConversionFlags))
 			{
 				TSharedRef<FJsonValueObject> JsonObject = MakeShared<FJsonValueObject>(Out);
 				JsonObject->Type = EJson::Object;
@@ -222,22 +267,22 @@ TSharedPtr<FJsonValue> ConvertScalarFPropertyToJsonValueWithContainer(FProperty*
 	return TSharedPtr<FJsonValue>();
 }
 
-TSharedPtr<FJsonValue> FPropertyToJsonValueWithContainer(FProperty* Property, const void* Value, const UObject* Container, int64 CheckFlags, int64 SkipFlags, const FJsonObjectConverter::CustomExportCallback* ExportCb, FProperty* OuterProperty, EJsonObjectConversionFlags ConversionFlags)
+TSharedPtr<FJsonValue> FPropertyToJsonValueWithContainer(FProperty* Property, const void* Value, const UObject* Container, TSet<const UObject*>* ExportedObjects, int64 CheckFlags, int64 SkipFlags, const FJsonObjectConverter::CustomExportCallback* ExportCb, FProperty* OuterProperty, EJsonObjectConversionFlags ConversionFlags)
 {
 	if (Property->ArrayDim == 1)
 	{
-		return ConvertScalarFPropertyToJsonValueWithContainer(Property, Value, Container, CheckFlags, SkipFlags, ExportCb, OuterProperty, ConversionFlags);
+		return ConvertScalarFPropertyToJsonValueWithContainer(Property, Value, Container, ExportedObjects, CheckFlags, SkipFlags, ExportCb, OuterProperty, ConversionFlags);
 	}
 
 	TArray< TSharedPtr<FJsonValue> > Array;
 	for (int Index = 0; Index != Property->ArrayDim; ++Index)
 	{
-		Array.Add(ConvertScalarFPropertyToJsonValueWithContainer(Property, (char*)Value + Index * Property->GetElementSize(), Container, CheckFlags, SkipFlags, ExportCb, OuterProperty, ConversionFlags));
+		Array.Add(ConvertScalarFPropertyToJsonValueWithContainer(Property, (char*)Value + Index * Property->GetElementSize(), Container, ExportedObjects, CheckFlags, SkipFlags, ExportCb, OuterProperty, ConversionFlags));
 	}
 	return MakeShared<FJsonValueArray>(Array);
 }
 
-bool UStructToJsonAttributesWithContainer(const UStruct* StructDefinition, const void* Struct, const UObject* Container, TMap< FString, TSharedPtr<FJsonValue> >& OutJsonAttributes, int64 CheckFlags, int64 SkipFlags, const FJsonObjectConverter::CustomExportCallback* ExportCb, EJsonObjectConversionFlags ConversionFlags)
+bool UStructToJsonAttributesWithContainer(const UStruct* StructDefinition, const void* Struct, const UObject* Container, TSet<const UObject*>* ExportedObjects, TMap< FString, TSharedPtr<FJsonValue> >& OutJsonAttributes, int64 CheckFlags, int64 SkipFlags, const FJsonObjectConverter::CustomExportCallback* ExportCb, EJsonObjectConversionFlags ConversionFlags)
 {
 	if (SkipFlags == 0)
 	{
@@ -280,7 +325,7 @@ bool UStructToJsonAttributesWithContainer(const UStruct* StructDefinition, const
 		const void* Value = Property->ContainerPtrToValuePtr<uint8>(Struct);
 
 		// convert the property to a FJsonValue
-		TSharedPtr<FJsonValue> JsonValue = FPropertyToJsonValueWithContainer(Property, Value, Container, CheckFlags, SkipFlags, ExportCb, nullptr, ConversionFlags);
+		TSharedPtr<FJsonValue> JsonValue = FPropertyToJsonValueWithContainer(Property, Value, Container, ExportedObjects, CheckFlags, SkipFlags, ExportCb, nullptr, ConversionFlags);
 		if (!JsonValue.IsValid())
 		{
 			FFieldClass* PropClass = Property->GetClass();
@@ -298,7 +343,7 @@ bool UStructToJsonAttributesWithContainer(const UStruct* StructDefinition, const
 
 TSharedPtr<FJsonValue> FJsonObjectConverter::UPropertyToJsonValue(FProperty* Property, const void* Value, int64 CheckFlags, int64 SkipFlags, const CustomExportCallback* ExportCb, FProperty* OuterProperty, EJsonObjectConversionFlags ConversionFlags)
 {
-	return FPropertyToJsonValueWithContainer(Property, Value, nullptr, CheckFlags, SkipFlags, ExportCb, OuterProperty, ConversionFlags);
+	return FPropertyToJsonValueWithContainer(Property, Value, nullptr, nullptr, CheckFlags, SkipFlags, ExportCb, OuterProperty, ConversionFlags);
 }
 
 bool FJsonObjectConverter::UStructToJsonObject(const UStruct* StructDefinition, const void* Struct, TSharedRef<FJsonObject> OutJsonObject, int64 CheckFlags, int64 SkipFlags, const CustomExportCallback* ExportCb, EJsonObjectConversionFlags ConversionFlags)
@@ -314,7 +359,8 @@ bool FJsonObjectConverter::UStructToJsonAttributes(const UStruct* StructDefiniti
 		ContainerObject = static_cast<const UObject*>(Struct);
 	}
 
-	return UStructToJsonAttributesWithContainer(StructDefinition, Struct, ContainerObject, OutJsonAttributes, CheckFlags, SkipFlags, ExportCb, ConversionFlags);
+	TSet<const UObject*> ExportedObjects;
+	return UStructToJsonAttributesWithContainer(StructDefinition, Struct, ContainerObject, &ExportedObjects, OutJsonAttributes, CheckFlags, SkipFlags, ExportCb, ConversionFlags);
 }
 
 template<class CharType, class PrintPolicy>
