Missing include for FPlatformStackWalk.

diff --git a/Engine/Source/Runtime/JsonUtilities/Private/JsonUClassExporter.cpp b/Engine/Source/Runtime/JsonUtilities/Private/JsonUClassExporter.cpp
--- a/Engine/Source/Runtime/JsonUtilities/Private/JsonUClassExporter.cpp
+++ b/Engine/Source/Runtime/JsonUtilities/Private/JsonUClassExporter.cpp
@@ -3,6 +3,7 @@
 #if !UE_BUILD_SHIPPING
 
 #include "HAL/IConsoleManager.h"
+#include "HAL/PlatformStackWalk.h"
 #include "Misc/FileHelper.h"
 #include "Misc/Paths.h"
 #include "Misc/SecureHash.h"
