[JsonObjectConverter] Adding static-interface and type-registration for custom struct converters. This allows structs to override the default json converter behavior.

diff --git a/Engine/Source/Runtime/JsonUtilities/Private/JsonObjectConverter.cpp b/Engine/Source/Runtime/JsonUtilities/Private/JsonObjectConverter.cpp
--- a/Engine/Source/Runtime/JsonUtilities/Private/JsonObjectConverter.cpp
+++ b/Engine/Source/Runtime/JsonUtilities/Private/JsonObjectConverter.cpp
@@ -1,6 +1,7 @@
 // Copyright Epic Games, Inc. All Rights Reserved.
 
 #include "JsonObjectConverter.h"
+#include "JsonObjectStructInterface.h"
 #include "Internationalization/Culture.h"
 #include "Misc/PackageName.h"
 #include "UObject/ObjectMacros.h"
@@ -207,13 +208,50 @@ TSharedPtr<FJsonValue> ConvertScalarFPropertyToJsonValueWithContainer(FProperty*
 	}
 	else if (FStructProperty *StructProperty = CastField<FStructProperty>(Property))
 	{
-		UScriptStruct::ICppStructOps* TheCppStructOps = StructProperty->Struct->GetCppStructOps();
 		// Intentionally exclude the JSON Object wrapper, which specifically needs to export JSON in an object representation instead of a string
-		if (StructProperty->Struct != FJsonObjectWrapper::StaticStruct() && TheCppStructOps && TheCppStructOps->HasExportTextItem())
+		if (StructProperty->Struct != FJsonObjectWrapper::StaticStruct())
 		{
-			FString OutValueStr;
-			TheCppStructOps->ExportTextItem(OutValueStr, Value, nullptr, nullptr, PPF_None, nullptr);
-			return MakeShared<FJsonValueString>(OutValueStr);
+			if (const IJsonObjectStructConverter* ConverterInterface = UE::Json::Private::GetStructConverterInterface(StructProperty->Struct))
+			{
+				TSharedPtr<FJsonObject> OutFromConverter = MakeShared<FJsonObject>();
+				switch (ConverterInterface->ConvertToJson(Value, OutFromConverter))
+				{
+				case EJsonObjectConvertResult::UseDefaultConverter:
+				{
+					// Continue to default converter.
+					break;
+				}
+				case EJsonObjectConvertResult::FailAndAbort:
+				{
+					// A failed ConvertToJson(...) gets treated as a null TSharedPtr, which can be seen as an error and may also abort the entire conversion.
+					// For example, UStructToJsonAttributesWithContainer will check JsonValue.IsValid() and treat invalid values as errors. 
+					return TSharedPtr<FJsonValue>();
+				}
+				case EJsonObjectConvertResult::IgnoreAndContinue:
+				{
+					// An ignored ConvertToJson(...) gets treated as an empty object (i.e. '{}'), but will not fail the conversion tree. 
+					return MakeShared<FJsonValueObject>(MakeShared<FJsonObject>());
+				}
+				case EJsonObjectConvertResult::Converted:
+				{
+					return MakeShared<FJsonValueObject>(OutFromConverter);
+				}
+				default:
+				{
+					unimplemented();
+					break;
+				}
+				}
+			}
+
+			// Default converter behavior is to use ExportTextItem if the struct has that function.
+			UScriptStruct::ICppStructOps* TheCppStructOps = StructProperty->Struct->GetCppStructOps();
+			if (TheCppStructOps && TheCppStructOps->HasExportTextItem())
+			{
+				FString OutValueStr;
+				TheCppStructOps->ExportTextItem(OutValueStr, Value, nullptr, nullptr, PPF_None, nullptr);
+				return MakeShared<FJsonValueString>(OutValueStr);
+			}
 		}
 
 		TSharedRef<FJsonObject> Out = MakeShared<FJsonObject>();
@@ -730,14 +768,53 @@ namespace
 			{
 				TSharedPtr<FJsonObject> Obj = JsonValue->AsObject();
 				check(Obj.IsValid()); // should not fail if Type == EJson::Object
-				if (!JsonAttributesToUStructWithContainer(Obj->Values, StructProperty->Struct, OutValue, ContainerStruct, Container, CheckFlags & (~CPF_ParmFlags), SkipFlags, bStrictMode, OutFailReason, ImportCb))
+
+				bool bUseDefaultConverter = true;
+
+				if (const IJsonObjectStructConverter* ConverterInterface = UE::Json::Private::GetStructConverterInterface(StructProperty->Struct))
 				{
-					UE_LOG(LogJson, Error, TEXT("JsonValueToUProperty - Unable to import JSON object into %s property %s"), *StructProperty->Struct->GetAuthoredName(), *Property->GetAuthoredName());
-					if (OutFailReason)
+					switch (ConverterInterface->ConvertFromJson(OutValue, Obj))
+					{
+					case EJsonObjectConvertResult::UseDefaultConverter:
 					{
-						*OutFailReason = FText::Format(LOCTEXT("FailImportStructFromObject", "Unable to import JSON object into {0} property {1}\n{2}"), FText::FromString(StructProperty->Struct->GetAuthoredName()), FText::FromString(Property->GetAuthoredName()), *OutFailReason);
+						bUseDefaultConverter = true;
+						break;
+					}
+					case EJsonObjectConvertResult::FailAndAbort:
+					{
+						UE_LOG(LogJson, Error, TEXT("JsonValueToUProperty - Unable to convert JSON object into %s property %s while use a IJsonObjectStructConverter"), *StructProperty->Struct->GetAuthoredName(), *Property->GetAuthoredName());
+						if (OutFailReason)
+						{
+							*OutFailReason = FText::Format(LOCTEXT("FailConvertStructFromObject", "Unable to convert JSON object into {0} property {1}\n{2} while using a IJsonObjectStructConverter"), FText::FromString(StructProperty->Struct->GetAuthoredName()), FText::FromString(Property->GetAuthoredName()), *OutFailReason);
+						}
+						return false;
+					}
+					case EJsonObjectConvertResult::IgnoreAndContinue:
+					case EJsonObjectConvertResult::Converted:
+					{
+						// IgnoreAndContinue and Converted are treated the same here
+						bUseDefaultConverter = false;
+						break;
+					}
+					default:
+					{
+						unimplemented();
+						break;
+					}
+					}
+				}
+
+				if (bUseDefaultConverter)
+				{
+					if (!JsonAttributesToUStructWithContainer(Obj->Values, StructProperty->Struct, OutValue, ContainerStruct, Container, CheckFlags & (~CPF_ParmFlags), SkipFlags, bStrictMode, OutFailReason, ImportCb))
+					{
+						UE_LOG(LogJson, Error, TEXT("JsonValueToUProperty - Unable to import JSON object into %s property %s"), *StructProperty->Struct->GetAuthoredName(), *Property->GetAuthoredName());
+						if (OutFailReason)
+						{
+							*OutFailReason = FText::Format(LOCTEXT("FailImportStructFromObject", "Unable to import JSON object into {0} property {1}\n{2}"), FText::FromString(StructProperty->Struct->GetAuthoredName()), FText::FromString(Property->GetAuthoredName()), *OutFailReason);
+						}
+						return false;
 					}
-					return false;
 				}
 			}
 			else if (JsonValue->Type == EJson::String && StructProperty->Struct->GetFName() == NAME_LinearColor)
diff --git a/Engine/Source/Runtime/JsonUtilities/Private/JsonObjectStructInterface.cpp b/Engine/Source/Runtime/JsonUtilities/Private/JsonObjectStructInterface.cpp
new file mode 100644
--- /dev/null
+++ b/Engine/Source/Runtime/JsonUtilities/Private/JsonObjectStructInterface.cpp
@@ -0,0 +1,33 @@
+// Copyright Epic Games, Inc. All Rights Reserved.
+
+#include "JsonObjectStructInterface.h"
+
+#include "Async/Mutex.h"
+#include "Async/UniqueLock.h"
+
+namespace JsonObjectStructInterface
+{
+	static UE::FMutex GStructConvertersMapMutex;
+	static TMap<const UScriptStruct*, const IJsonObjectStructConverter*> GStructConvertersMap;
+}
+
+void FJsonObjectStructInterfaceRegistry::RegisterStructConverter(const UScriptStruct* ScriptStruct, const IJsonObjectStructConverter* ConverterInterface)
+{
+	UE::TUniqueLock Lock(JsonObjectStructInterface::GStructConvertersMapMutex);
+	JsonObjectStructInterface::GStructConvertersMap.Add(ScriptStruct, ConverterInterface);
+}
+
+void FJsonObjectStructInterfaceRegistry::UnregisterStructConverter(const UScriptStruct* ScriptStruct)
+{
+	UE::TUniqueLock Lock(JsonObjectStructInterface::GStructConvertersMapMutex);
+	JsonObjectStructInterface::GStructConvertersMap.Remove(ScriptStruct);
+}
+
+namespace UE::Json::Private
+{
+	const IJsonObjectStructConverter* GetStructConverterInterface(const UScriptStruct* ScriptStruct)
+	{
+		UE::TUniqueLock Lock(JsonObjectStructInterface::GStructConvertersMapMutex);
+		return JsonObjectStructInterface::GStructConvertersMap.FindRef(ScriptStruct);
+	}
+}
\ No newline at end of file
diff --git a/Engine/Source/Runtime/JsonUtilities/Public/JsonObjectStructInterface.h b/Engine/Source/Runtime/JsonUtilities/Public/JsonObjectStructInterface.h
new file mode 100644
--- /dev/null
+++ b/Engine/Source/Runtime/JsonUtilities/Public/JsonObjectStructInterface.h
@@ -0,0 +1,97 @@
+// Copyright Epic Games, Inc. All Rights Reserved.
+
+#pragma once
+
+#include "UObject/Class.h"
+
+class FJsonObject;
+struct IJsonObjectStructConverter;
+
+/** Global registry for type-registered static-interface converters to override default struct behavior in FJsonObjectConverter. */
+struct FJsonObjectStructInterfaceRegistry
+{
+	static JSONUTILITIES_API void RegisterStructConverter(const UScriptStruct* ScriptStruct, const IJsonObjectStructConverter* ConverterInterface);
+	static JSONUTILITIES_API void UnregisterStructConverter(const UScriptStruct* ScriptStruct);
+};
+
+enum class EJsonObjectConvertResult
+{
+	/** No conversion was performed. Fallback to try the default JsonObjectConverter instead. */
+	UseDefaultConverter,
+	/** No conversion is possible. Abandon further efforts. Ask caller to fail. */
+	FailAndAbort,
+	/** No conversion was performed. Don't try any fallback. Allow caller to still succeed the conversion (if this was a nested convert attempt). */
+	IgnoreAndContinue,
+	/** Conversion of the struct was performed. */
+	Converted,
+};
+
+/**
+ * UStructs are special and may need to avoid the use of virtual functions via standard interface patterns due to mismatched base-class polymorphism.
+ * So use a static TImplementsJsonObjectStructConverter<StructType> which uses a static-interface pattern with global type-registry instead of implementing an interface directly.
+ *
+ * The registered struct should have the following functions (the compiler will fail without them):
+ * EJsonObjectConvertResult ConvertToJson(TSharedPtr<FJsonObject>& OutJsonObject) const;
+ * EJsonObjectConvertResult ConvertFromJson(const TSharedPtr<FJsonObject>& InJsonObject);
+ */
+struct IJsonObjectStructConverter
+{
+	/** 
+	* Result of Converted represents success by the interface call,
+	* FailAndAbort will bubble up the failure to its owner converter which may fail the entire converter tree,
+	* IgnoreAndContinue will make an empty FJsonObject for this struct and the converter tree can still succeed,
+	* and UseDefaultConverter will try to fallback to the non-static-interface behavior instead.
+	*/
+	virtual EJsonObjectConvertResult ConvertToJson(const void* StructMemory, TSharedPtr<FJsonObject>& OutJsonObject) const = 0;
+
+	/** 
+	* Result of Converted represents success by the interface call,
+	* FailAndAbort will bubble up the failure to its owner converter failing the entire converter tree, 
+	* IgnoreAndContinue should leave the struct untouched/default and the converter tree can still succeed, 
+	* and UseDefaultConverter will try to fallback to the non-static-interface behavior instead. 
+	*/
+	virtual EJsonObjectConvertResult ConvertFromJson(void* StructMemory, const TSharedPtr<FJsonObject>& InJsonObject) const = 0;
+
+protected:
+
+	virtual ~IJsonObjectStructConverter() = default;
+};
+
+/** 
+* This template is to be statically constructed and passed as a parameter to RegisterStructConverter in StartupModule() for the specified StructType. 
+* Remember to explicitly call UnregisterStructConverter in ShutdownModule() as well.
+* 
+* Example usage:
+* 
+* void FExampleModule::StartupModule()
+* {
+* 	static const TImplementsJsonObjectStructConverter<FExampleStruct> ExampleConverter = TImplementsJsonObjectStructConverter<FExampleStruct>();
+* 	FJsonObjectStructInterfaceRegistry::RegisterStructConverter(FExampleStruct::StaticStruct(), &ExampleConverter);
+* }
+* 
+* void FExampleModule::ShutdownModule()
+* {
+* 	FJsonObjectStructInterfaceRegistry::UnregisterStructConverter(FExampleStruct::StaticStruct());
+* }
+*/
+template<typename StructType>
+struct TImplementsJsonObjectStructConverter : public IJsonObjectStructConverter
+{
+	TImplementsJsonObjectStructConverter<StructType>() {};
+
+	virtual EJsonObjectConvertResult ConvertToJson(const void* StructMemory, TSharedPtr<FJsonObject>& OutJsonObject) const override
+	{
+		return static_cast<const StructType*>(StructMemory)->ConvertToJson(OutJsonObject);
+	}
+
+	virtual EJsonObjectConvertResult ConvertFromJson(void* StructMemory, const TSharedPtr<FJsonObject>& InJsonObject) const override
+	{
+		return static_cast<StructType*>(StructMemory)->ConvertFromJson(InJsonObject);
+	}
+};
+
+/** Functions for internal use. */
+namespace UE::Json::Private
+{
+	const struct IJsonObjectStructConverter* GetStructConverterInterface(const UScriptStruct* ScriptStruct);
+}
\ No newline at end of file
