[Backout] - CL44932745
#rnx
Original CL Desc
-----------------------------------------------------------------
[JsonObjectConverter] Adding static-interface and type-registration for custom struct converters. This allows structs to override the default json converter behavior.

diff --git a/Engine/Source/Runtime/JsonUtilities/Private/JsonObjectConverter.cpp b/Engine/Source/Runtime/JsonUtilities/Private/JsonObjectConverter.cpp
--- a/Engine/Source/Runtime/JsonUtilities/Private/JsonObjectConverter.cpp
+++ b/Engine/Source/Runtime/JsonUtilities/Private/JsonObjectConverter.cpp
@@ -1,7 +1,6 @@
 // Copyright Epic Games, Inc. All Rights Reserved.
 
 #include "JsonObjectConverter.h"
-#include "JsonObjectStructInterface.h"
 #include "Internationalization/Culture.h"
 #include "Misc/PackageName.h"
 #include "UObject/ObjectMacros.h"
@@ -208,50 +207,13 @@ TSharedPtr<FJsonValue> ConvertScalarFPropertyToJsonValueWithContainer(FProperty*
 	}
 	else if (FStructProperty *StructProperty = CastField<FStructProperty>(Property))
 	{
+		UScriptStruct::ICppStructOps* TheCppStructOps = StructProperty->Struct->GetCppStructOps();
 		// Intentionally exclude the JSON Object wrapper, which specifically needs to export JSON in an object representation instead of a string
-		if (StructProperty->Struct != FJsonObjectWrapper::StaticStruct())
+		if (StructProperty->Struct != FJsonObjectWrapper::StaticStruct() && TheCppStructOps && TheCppStructOps->HasExportTextItem())
 		{
-			if (const IJsonObjectStructConverter* ConverterInterface = UE::Json::Private::GetStructConverterInterface(StructProperty->Struct))
-			{
-				TSharedPtr<FJsonObject> OutFromConverter = MakeShared<FJsonObject>();
-				switch (ConverterInterface->ConvertToJson(Value, OutFromConverter))
-				{
-				case EJsonObjectConvertResult::UseDefaultConverter:
-				{
-					// Continue to default converter.
-					break;
-				}
-				case EJsonObjectConvertResult::FailAndAbort:
-				{
-					// A failed ConvertToJson(...) gets treated as a null TSharedPtr, which can be seen as an error and may also abort the entire conversion.
-					// For example, UStructToJsonAttributesWithContainer will check JsonValue.IsValid() and treat invalid values as errors. 
-					return TSharedPtr<FJsonValue>();
-				}
-				case EJsonObjectConvertResult::IgnoreAndContinue:
-				{
-					// An ignored ConvertToJson(...) gets treated as an empty object (i.e. '{}'), but will not fail the conversion tree. 
-					return MakeShared<FJsonValueObject>(MakeShared<FJsonObject>());
-				}
-				case EJsonObjectConvertResult::Converted:
-				{
-					return MakeShared<FJsonValueObject>(OutFromConverter);
-				}
-				default:
-				{
-					unimplemented();
-					break;
-				}
-				}
-			}
-
-			// Default converter behavior is to use ExportTextItem if the struct has that function.
-			UScriptStruct::ICppStructOps* TheCppStructOps = StructProperty->Struct->GetCppStructOps();
-			if (TheCppStructOps && TheCppStructOps->HasExportTextItem())
-			{
-				FString OutValueStr;
-				TheCppStructOps->ExportTextItem(OutValueStr, Value, nullptr, nullptr, PPF_None, nullptr);
-				return MakeShared<FJsonValueString>(OutValueStr);
-			}
+			FString OutValueStr;
+			TheCppStructOps->ExportTextItem(OutValueStr, Value, nullptr, nullptr, PPF_None, nullptr);
+			return MakeShared<FJsonValueString>(OutValueStr);
 		}
 
 		TSharedRef<FJsonObject> Out = MakeShared<FJsonObject>();
@@ -768,53 +730,14 @@ namespace
 			{
 				TSharedPtr<FJsonObject> Obj = JsonValue->AsObject();
 				check(Obj.IsValid()); // should not fail if Type == EJson::Object
-
-				bool bUseDefaultConverter = true;
-
-				if (const IJsonObjectStructConverter* ConverterInterface = UE::Json::Private::GetStructConverterInterface(StructProperty->Struct))
-				{
-					switch (ConverterInterface->ConvertFromJson(OutValue, Obj))
-					{
-					case EJsonObjectConvertResult::UseDefaultConverter:
-					{
-						bUseDefaultConverter = true;
-						break;
-					}
-					case EJsonObjectConvertResult::FailAndAbort:
-					{
-						UE_LOG(LogJson, Error, TEXT("JsonValueToUProperty - Unable to convert JSON object into %s property %s while use a IJsonObjectStructConverter"), *StructProperty->Struct->GetAuthoredName(), *Property->GetAuthoredName());
-						if (OutFailReason)
-						{
-							*OutFailReason = FText::Format(LOCTEXT("FailConvertStructFromObject", "Unable to convert JSON object into {0} property {1}\n{2} while using a IJsonObjectStructConverter"), FText::FromString(StructProperty->Struct->GetAuthoredName()), FText::FromString(Property->GetAuthoredName()), *OutFailReason);
-						}
-						return false;
-					}
-					case EJsonObjectConvertResult::IgnoreAndContinue:
-					case EJsonObjectConvertResult::Converted:
-					{
-						// IgnoreAndContinue and Converted are treated the same here
-						bUseDefaultConverter = false;
-						break;
-					}
-					default:
-					{
-						unimplemented();
-						break;
-					}
-					}
-				}
-
-				if (bUseDefaultConverter)
+				if (!JsonAttributesToUStructWithContainer(Obj->Values, StructProperty->Struct, OutValue, ContainerStruct, Container, CheckFlags & (~CPF_ParmFlags), SkipFlags, bStrictMode, OutFailReason, ImportCb))
 				{
-					if (!JsonAttributesToUStructWithContainer(Obj->Values, StructProperty->Struct, OutValue, ContainerStruct, Container, CheckFlags & (~CPF_ParmFlags), SkipFlags, bStrictMode, OutFailReason, ImportCb))
+					UE_LOG(LogJson, Error, TEXT("JsonValueToUProperty - Unable to import JSON object into %s property %s"), *StructProperty->Struct->GetAuthoredName(), *Property->GetAuthoredName());
+					if (OutFailReason)
 					{
-						UE_LOG(LogJson, Error, TEXT("JsonValueToUProperty - Unable to import JSON object into %s property %s"), *StructProperty->Struct->GetAuthoredName(), *Property->GetAuthoredName());
-						if (OutFailReason)
-						{
-							*OutFailReason = FText::Format(LOCTEXT("FailImportStructFromObject", "Unable to import JSON object into {0} property {1}\n{2}"), FText::FromString(StructProperty->Struct->GetAuthoredName()), FText::FromString(Property->GetAuthoredName()), *OutFailReason);
-						}
-						return false;
+						*OutFailReason = FText::Format(LOCTEXT("FailImportStructFromObject", "Unable to import JSON object into {0} property {1}\n{2}"), FText::FromString(StructProperty->Struct->GetAuthoredName()), FText::FromString(Property->GetAuthoredName()), *OutFailReason);
 					}
+					return false;
 				}
 			}
 			else if (JsonValue->Type == EJson::String && StructProperty->Struct->GetFName() == NAME_LinearColor)
diff --git a/Engine/Source/Runtime/JsonUtilities/Private/JsonObjectStructInterface.cpp b/Engine/Source/Runtime/JsonUtilities/Private/JsonObjectStructInterface.cpp
deleted file mode 100644
diff --git a/Engine/Source/Runtime/JsonUtilities/Public/JsonObjectStructInterface.h b/Engine/Source/Runtime/JsonUtilities/Public/JsonObjectStructInterface.h
deleted file mode 100644
