[JsonObjectConverter] Added support for TOption properties. Unset TOption properties will now be skipped in JSON export & their inner values correctly exported when set. Conversely, JSON import will now correctly set the inner property values.

Previously, TOption properties (FOptionProperty) would be exported and imported as text, losing their rich inner typing, where TOption can now be used in structs for I/O of optional json fields.

diff --git a/Engine/Source/Runtime/JsonUtilities/Private/JsonObjectConverter.cpp b/Engine/Source/Runtime/JsonUtilities/Private/JsonObjectConverter.cpp
--- a/Engine/Source/Runtime/JsonUtilities/Private/JsonObjectConverter.cpp
+++ b/Engine/Source/Runtime/JsonUtilities/Private/JsonObjectConverter.cpp
@@ -10,6 +10,7 @@
 #include "UObject/UnrealType.h"
 #include "UObject/EnumProperty.h"
 #include "UObject/TextProperty.h"
+#include "UObject/PropertyOptional.h"
 #include "UObject/PropertyPortFlags.h"
 #include "UObject/Package.h"
 #include "Policies/CondensedJsonPrintPolicy.h"
@@ -294,6 +295,20 @@ TSharedPtr<FJsonValue> ConvertScalarFPropertyToJsonValueWithContainer(FProperty*
 			return MakeShared<FJsonValueString>(StringValue);
 		}
 	}
+	else if (FOptionalProperty* OptionalProperty = CastField<FOptionalProperty>(Property))
+	{
+		if (const void* InnerValue = OptionalProperty->GetValuePointerForReadIfSet(Value))
+		{
+			return ConvertScalarFPropertyToJsonValueWithContainer(OptionalProperty->GetValueProperty(), InnerValue, Container, ExportedObjects, CheckFlags, SkipFlags, ExportCb, OptionalProperty, ConversionFlags);
+		}
+		else
+		{
+			check(!OptionalProperty->IsSet(Value));
+			// Note: Unset optional properties are skipped by UStructToJsonAttributesWithContainer, however custom export callbacks or explicit calls to 
+			// UPropertyToJsonValue may still arrive here.
+			return MakeShared<FJsonValueNull>();
+		}
+	}
 	else
 	{
 		// Default to export as string for everything else
@@ -354,6 +369,17 @@ bool UStructToJsonAttributesWithContainer(const UStruct* StructDefinition, const
 		{
 			continue;
 		}
+		
+		const void* Value = Property->ContainerPtrToValuePtr<uint8>(Struct);
+
+		// Skip unset optional properties
+		if (const FOptionalProperty* OptionalProperty = CastField<FOptionalProperty>(Property))
+		{
+			if (!OptionalProperty->IsSet(Value))
+			{
+				continue;
+			}
+		}
 
 		FString VariableName = Property->GetAuthoredName();
 		if (!EnumHasAnyFlags(ConversionFlags, EJsonObjectConversionFlags::SkipStandardizeCase))
@@ -361,8 +387,6 @@ bool UStructToJsonAttributesWithContainer(const UStruct* StructDefinition, const
 			VariableName = FJsonObjectConverter::StandardizeCase(VariableName);
 		}
 
-		const void* Value = Property->ContainerPtrToValuePtr<uint8>(Struct);
-
 		// convert the property to a FJsonValue
 		TSharedPtr<FJsonValue> JsonValue = FPropertyToJsonValueWithContainer(Property, Value, Container, ExportedObjects, CheckFlags, SkipFlags, ExportCb, nullptr, ConversionFlags);
 		if (!JsonValue.IsValid())
@@ -970,6 +994,11 @@ namespace
 				}
 			}
 		}
+		else if (FOptionalProperty* OptionalProperty = CastField<FOptionalProperty>(Property))
+		{
+			void* OutValueInner = OptionalProperty->MarkSetAndGetInitializedValuePointerToReplace(OutValue);
+			return ConvertScalarJsonValueToFPropertyWithContainer(JsonValue, OptionalProperty->GetValueProperty(), OutValueInner, ContainerStruct, Container, CheckFlags, SkipFlags, bStrictMode, OutFailReason, ImportCb);
+		}
 		else
 		{
 			// Default to expect a string for everything else
