JSON Utilities: fix illegal instantiation of objects of abstract class.

diff --git a/Engine/Source/Runtime/JsonUtilities/Private/JsonObjectConverter.cpp b/Engine/Source/Runtime/JsonUtilities/Private/JsonObjectConverter.cpp
--- a/Engine/Source/Runtime/JsonUtilities/Private/JsonObjectConverter.cpp
+++ b/Engine/Source/Runtime/JsonUtilities/Private/JsonObjectConverter.cpp
@@ -973,6 +973,26 @@ namespace
 					{
 						PropertyClass = FoundClass;
 					}
+					else
+					{
+						UE_LOG(LogJson, Error, TEXT("JsonValueToUProperty - Unable to import JSON object for %s property %s. Class %s was not found."), *PropertyClass->GetAuthoredName(), *Property->GetAuthoredName(), *ClassString);
+						if (OutFailReason)
+						{
+							*OutFailReason = FText::Format(LOCTEXT("FailImportObjectClassNotFound", "Unable to import JSON object into {0} property {1}. Class {2} was not found\n{3}"), FText::FromString(PropertyClass->GetAuthoredName()), FText::FromString(Property->GetAuthoredName()), FText::FromString(ClassString), *OutFailReason);
+						}
+						// Continue execution to keep consistency with original behavior, but log error.
+					}
+				}
+
+				// Abstract classes should not be instantiated 
+				if (PropertyClass->HasAnyClassFlags(CLASS_Abstract))
+				{
+					UE_LOG(LogJson, Error, TEXT("JsonValueToUProperty - Unable to import JSON object into %s property %s. Class is abstract."), *PropertyClass->GetAuthoredName(), *Property->GetAuthoredName());
+					if (OutFailReason)
+					{
+						*OutFailReason = FText::Format(LOCTEXT("FailImportObjectAbstractClass", "Unable to import JSON object into {0} property {1}. Class is abstract\n{2}"), FText::FromString(PropertyClass->GetAuthoredName()), FText::FromString(Property->GetAuthoredName()), *OutFailReason);
+					}
+					return false;
 				}
 
 				UObject* createdObj = StaticAllocateObject(PropertyClass, Outer, NAME_None, EObjectFlags::RF_NoFlags, EInternalObjectFlags::None, false);
