JSON Utilities: Guarantee that the Map Helper always calls rehash even when bailing out as some items could've already been added.

diff --git a/Engine/Source/Runtime/JsonUtilities/Private/JsonObjectConverter.cpp b/Engine/Source/Runtime/JsonUtilities/Private/JsonObjectConverter.cpp
--- a/Engine/Source/Runtime/JsonUtilities/Private/JsonObjectConverter.cpp
+++ b/Engine/Source/Runtime/JsonUtilities/Private/JsonObjectConverter.cpp
@@ -5,6 +5,7 @@
 #include "JsonObjectWrapper.h"
 #include "Internationalization/Culture.h"
 #include "Misc/PackageName.h"
+#include "Misc/ScopeExit.h"
 #include "UObject/ObjectMacros.h"
 #include "UObject/Class.h"
 #include "UObject/StrProperty.h"
@@ -671,6 +672,11 @@ namespace
 				int32 MapSize = ObjectValue->Values.Num();
 				Helper.EmptyValues(MapSize);
 
+				ON_SCOPE_EXIT
+				{
+					Helper.Rehash();
+				};
+
 				// set the property values
 				for (const auto& Entry : ObjectValue->Values)
 				{
@@ -701,8 +707,6 @@ namespace
 						}
 					}
 				}
-
-				Helper.Rehash();
 			}
 			else
 			{
@@ -724,6 +728,11 @@ namespace
 				FScriptSetHelper Helper(SetProperty, OutValue);
 				Helper.EmptyElements(ArrLen);
 
+				ON_SCOPE_EXIT
+				{
+					Helper.Rehash();
+				};
+
 				// set the property values
 				for (int32 i = 0; i < ArrLen; ++i)
 				{
@@ -742,8 +751,6 @@ namespace
 						}
 					}
 				}
-
-				Helper.Rehash();
 			}
 			else
 			{
