GitHub 13851 : Fix for "Field _ClassName was not found." warning

#13851

diff --git a/Engine/Source/Runtime/JsonUtilities/Private/JsonObjectConverter.cpp b/Engine/Source/Runtime/JsonUtilities/Private/JsonObjectConverter.cpp
--- a/Engine/Source/Runtime/JsonUtilities/Private/JsonObjectConverter.cpp
+++ b/Engine/Source/Runtime/JsonUtilities/Private/JsonObjectConverter.cpp
@@ -971,23 +971,26 @@ namespace
 				UClass* PropertyClass = ObjectProperty->PropertyClass;
 
 				// If a specific subclass was stored in the JSON, use that instead of the PropertyClass
-				FString ClassString = Obj->GetStringField(ObjectClassNameKey);
-				Obj->RemoveField(ObjectClassNameKey);
-				if (!ClassString.IsEmpty())
+				FString ClassString;
+				if (Obj->TryGetStringField(ObjectClassNameKey, ClassString))
 				{
-					UClass* FoundClass = FPackageName::IsShortPackageName(ClassString) ? FindFirstObject<UClass>(*ClassString) : LoadClass<UObject>(nullptr, *ClassString);
-					if (FoundClass)
+					Obj->RemoveField(ObjectClassNameKey);
+					if (!ClassString.IsEmpty())
 					{
-						PropertyClass = FoundClass;
-					}
-					else
-					{
-						UE_LOG(LogJson, Error, TEXT("JsonValueToUProperty - Unable to import JSON object for %s property %s. Class %s was not found."), *PropertyClass->GetAuthoredName(), *Property->GetAuthoredName(), *ClassString);
-						if (OutFailReason)
+						UClass* FoundClass = FPackageName::IsShortPackageName(ClassString) ? FindFirstObject<UClass>(*ClassString) : LoadClass<UObject>(nullptr, *ClassString);
+						if (FoundClass)
+						{
+							PropertyClass = FoundClass;
+						}
+						else
 						{
-							*OutFailReason = FText::Format(LOCTEXT("FailImportObjectClassNotFound", "Unable to import JSON object into {0} property {1}. Class {2} was not found\n{3}"), FText::FromString(PropertyClass->GetAuthoredName()), FText::FromString(Property->GetAuthoredName()), FText::FromString(ClassString), *OutFailReason);
+							UE_LOG(LogJson, Error, TEXT("JsonValueToUProperty - Unable to import JSON object for %s property %s. Class %s was not found."), *PropertyClass->GetAuthoredName(), *Property->GetAuthoredName(), *ClassString);
+							if (OutFailReason)
+							{
+								*OutFailReason = FText::Format(LOCTEXT("FailImportObjectClassNotFound", "Unable to import JSON object into {0} property {1}. Class {2} was not found\n{3}"), FText::FromString(PropertyClass->GetAuthoredName()), FText::FromString(Property->GetAuthoredName()), FText::FromString(ClassString), *OutFailReason);
+							}
+							// Continue execution to keep consistency with original behavior, but log error.
 						}
-						// Continue execution to keep consistency with original behavior, but log error.
 					}
 				}
 
