JSON Serialization optimizations

- New: FJsonDomBuilder::ToUtf8String(): Serialize directly to UTF-8 to avoid expensive wchar->utf8 conversion.
- New: FJsonObjectConverter::StandardizeCaseInPlace(): In-place string case folding; avoids a string allocation
- New: FJsonObject::SetObjectField() overload: accepts a shared pointer by rvalue reference to avoid the construction of an extra shared pointer when possible

diff --git a/Engine/Source/Runtime/Json/Private/Dom/JsonObject.cpp b/Engine/Source/Runtime/Json/Private/Dom/JsonObject.cpp
--- a/Engine/Source/Runtime/Json/Private/Dom/JsonObject.cpp
+++ b/Engine/Source/Runtime/Json/Private/Dom/JsonObject.cpp
@@ -242,18 +242,23 @@ bool FJsonObject::TryGetObjectField(FStringView FieldName, const TSharedPtr<FJso
 	return Field.IsValid() && Field->TryGetObject(OutObject);
 }
 
-void FJsonObject::SetObjectField(FString&& FieldName, const TSharedPtr<FJsonObject>& JsonObject)
+void FJsonObject::SetObjectField(FString&& FieldName, TSharedPtr<FJsonObject>&& JsonObject)
 {
 	if (JsonObject.IsValid())
 	{
-		this->Values.Add(MoveTemp(FieldName), MakeShared<FJsonValueObject>(JsonObject.ToSharedRef()));
+		this->Values.Emplace(MoveTemp(FieldName), MakeShared<FJsonValueObject>(MoveTemp(JsonObject)));
 	}
 	else
 	{
-		this->Values.Add(MoveTemp(FieldName), MakeShared<FJsonValueNull>());
+		this->Values.Emplace(MoveTemp(FieldName), MakeShared<FJsonValueNull>());
 	}
 }
 
+void FJsonObject::SetObjectField(FString&& FieldName, const TSharedPtr<FJsonObject>& JsonObject)
+{
+	SetObjectField(MoveTemp(FieldName), TSharedPtr<FJsonObject>(JsonObject));
+}
+
 void FJsonObject::SetObjectField(const FString& FieldName, const TSharedPtr<FJsonObject>& JsonObject)
 {
 	SetObjectField(CopyTemp(FieldName), JsonObject);
diff --git a/Engine/Source/Runtime/Json/Public/Dom/JsonObject.h b/Engine/Source/Runtime/Json/Public/Dom/JsonObject.h
--- a/Engine/Source/Runtime/Json/Public/Dom/JsonObject.h
+++ b/Engine/Source/Runtime/Json/Public/Dom/JsonObject.h
@@ -293,6 +293,7 @@ public:
 	JSON_API bool TryGetObjectField(FStringView FieldName, const TSharedPtr<FJsonObject>*& OutObject) const;
 
 	/** Set an ObjectField named FieldName and value of JsonObject */
+	JSON_API void SetObjectField(FString&& FieldName, TSharedPtr<FJsonObject>&& JsonObject);
 	JSON_API void SetObjectField(FString&& FieldName, const TSharedPtr<FJsonObject>& JsonObject);
 	JSON_API void SetObjectField(const FString& FieldName, const TSharedPtr<FJsonObject>& JsonObject);
 
diff --git a/Engine/Source/Runtime/JsonUtilities/Private/JsonObjectConverter.cpp b/Engine/Source/Runtime/JsonUtilities/Private/JsonObjectConverter.cpp
--- a/Engine/Source/Runtime/JsonUtilities/Private/JsonObjectConverter.cpp
+++ b/Engine/Source/Runtime/JsonUtilities/Private/JsonObjectConverter.cpp
@@ -23,13 +23,22 @@
 
 FString FJsonObjectConverter::StandardizeCase(const FString &StringIn)
 {
-	// this probably won't work for all cases, consider downcasing the string fully
 	FString FixedString = StringIn;
-	FixedString[0] = FChar::ToLower(FixedString[0]); // our JSON classes/variable start lower case
-	FixedString.ReplaceInline(TEXT("ID"), TEXT("Id"), ESearchCase::CaseSensitive); // Id is standard instead of ID, some of our fnames use ID
+	StandardizeCaseInPlace(FixedString);
 	return FixedString;
 }
 
+void FJsonObjectConverter::StandardizeCaseInPlace(FString& String)
+{
+	if (String.IsEmpty())
+	{
+		return;
+	}
+
+	// this probably won't work for all cases, consider downcasing the string fully
+	String[0] = FChar::ToLower(String[0]); // our JSON classes/variable start lower case
+	String.ReplaceInline(TEXT("ID"), TEXT("Id"), ESearchCase::CaseSensitive); // Id is standard instead of ID, some of our fnames use ID
+}
 
 namespace
 {
@@ -201,7 +210,7 @@ TSharedPtr<FJsonValue> ConvertScalarFPropertyToJsonValueWithContainer(FProperty*
 				{
 					if (!EnumHasAnyFlags(ConversionFlags, EJsonObjectConversionFlags::SkipStandardizeCase))
 					{
-						KeyString = FJsonObjectConverter::StandardizeCase(KeyString);
+						FJsonObjectConverter::StandardizeCaseInPlace(KeyString);
 					}
 				}
 				Out->SetField(KeyString, ValueElement);
@@ -396,7 +405,7 @@ bool UStructToJsonAttributesWithContainer(const UStruct* StructDefinition, const
 		FString VariableName = Property->GetAuthoredName();
 		if (!EnumHasAnyFlags(ConversionFlags, EJsonObjectConversionFlags::SkipStandardizeCase))
 		{
-			VariableName = FJsonObjectConverter::StandardizeCase(VariableName);
+			FJsonObjectConverter::StandardizeCaseInPlace(VariableName);
 		}
 
 		// convert the property to a FJsonValue
diff --git a/Engine/Source/Runtime/JsonUtilities/Public/JsonDomBuilder.h b/Engine/Source/Runtime/JsonUtilities/Public/JsonDomBuilder.h
--- a/Engine/Source/Runtime/JsonUtilities/Public/JsonDomBuilder.h
+++ b/Engine/Source/Runtime/JsonUtilities/Public/JsonDomBuilder.h
@@ -69,6 +69,15 @@ public:
 			return Result;
 		}
 
+		template <template <class> class TPrintPolicy = TPrettyJsonPrintPolicy>
+		FUtf8String ToUtf8String() const
+		{
+			FUtf8String Result;
+			auto JsonWriter = TJsonWriterFactory<UTF8CHAR, TPrintPolicy<UTF8CHAR>>::Create(&Result);
+			FJsonSerializer::Serialize(Object, JsonWriter);
+			return Result;
+		}
+
 		int Num() const
 		{
 			return Object->Values.Num();
diff --git a/Engine/Source/Runtime/JsonUtilities/Public/JsonObjectConverter.h b/Engine/Source/Runtime/JsonUtilities/Public/JsonObjectConverter.h
--- a/Engine/Source/Runtime/JsonUtilities/Public/JsonObjectConverter.h
+++ b/Engine/Source/Runtime/JsonUtilities/Public/JsonObjectConverter.h
@@ -56,8 +56,13 @@ class FJsonObjectConverter
 {
 public:
 
-	/** FName case insensitivity can make the casing of UPROPERTIES unpredictable. Attempt to standardize output. */
-	static JSONUTILITIES_API FString StandardizeCase(const FString &StringIn);
+	/**
+	 * FName case insensitivity can make the casing of UPROPERTIES unpredictable. Attempt to standardize output.
+	 * See StandardizeCaseInPlace() for an in-place equivalent.
+	 */
+	[[nodiscard]] static JSONUTILITIES_API FString StandardizeCase(const FString &StringIn);
+
+	static JSONUTILITIES_API void StandardizeCaseInPlace(FString& String);
 
 	/** Parse an FText from a json object (assumed to be of the form where keys are culture codes and values are strings) */
 	static JSONUTILITIES_API bool GetTextFromObject(const TSharedRef<FJsonObject>& Obj, FText& TextOut);
