[json] Improved and more consistent DOM builder argument conversions

#Synced-CL 48296048

diff --git a/Engine/Source/Runtime/JsonUtilities/Public/JsonDomBuilder.h b/Engine/Source/Runtime/JsonUtilities/Public/JsonDomBuilder.h
--- a/Engine/Source/Runtime/JsonUtilities/Public/JsonDomBuilder.h
+++ b/Engine/Source/Runtime/JsonUtilities/Public/JsonDomBuilder.h
@@ -86,15 +86,17 @@ public:
 		FObject& Set(const FString& Key, const FArray& Arr)            { Object->SetField(Key, Arr.AsJsonValue());                            return *this; }
 		FObject& Set(const FString& Key, const FObject& Obj)           { Object->SetField(Key, Obj.AsJsonValue());                            return *this; }
 
-		FObject& Set(const FString& Key, const FString& Str)           { Object->SetField(Key, MakeShared<FJsonValueString>(Str));            return *this; }
+		template <class StringType
+			UE_REQUIRES(std::is_constructible_v<FString, StringType>)>
+		FObject& Set(const FString& Key, const StringType& Str)        { Object->SetField(Key, MakeShared<FJsonValueString>(FString(Str)));   return *this; }
 
 		template <class NumberType
 			UE_REQUIRES(!std::is_same_v<NumberType, bool> && (std::is_integral_v<NumberType> || std::is_floating_point_v<NumberType>))>
-		FObject& Set(const FString& Key, NumberType Number)                    { Object->SetField(Key, MakeShared<FJsonValueNumber>(Number));         return *this; }
+		FObject& Set(const FString& Key, NumberType Number)            { Object->SetField(Key, MakeShared<FJsonValueNumber>(Number));         return *this; }
 
 		template <class BoolType
 			UE_REQUIRES(std::is_same_v<BoolType, bool>)>
-		FObject& Set(const FString& Key, BoolType Boolean)                     { Object->SetField(Key, MakeShared<FJsonValueBoolean>(Boolean));       return *this; }
+		FObject& Set(const FString& Key, BoolType Boolean)             { Object->SetField(Key, MakeShared<FJsonValueBoolean>(Boolean));       return *this; }
 
 		FObject& Set(const FString& Key, TYPE_OF_NULLPTR)              { Object->SetField(Key, MakeShared<FJsonValueNull>());                 return *this; }
 
@@ -132,6 +134,15 @@ public:
 			return Result;
 		}
 
+		template <template <class> class TPrintPolicy = TPrettyJsonPrintPolicy>
+		FUtf8String ToUtf8String() const
+		{
+			FUtf8String Result;
+			auto JsonWriter = TJsonWriterFactory<UTF8CHAR, TPrintPolicy<UTF8CHAR>>::Create(&Result);
+			FJsonSerializer::Serialize(Array, JsonWriter);
+			return Result;
+		}
+
 		int Num() const
 		{
 			return Array.Num();
@@ -140,13 +151,17 @@ public:
 		FArray& Add(const FArray& Arr)            { Array.Emplace(Arr.AsJsonValue());                      return *this; }
 		FArray& Add(const FObject& Obj)           { Array.Emplace(Obj.AsJsonValue());                      return *this; }
 
-		FArray& Add(const FString& Str)           { Array.Emplace(MakeShared<FJsonValueString>(Str));      return *this; }
+		template <class StringType
+			UE_REQUIRES(std::is_constructible_v<FString, StringType>)>
+		FArray& Add(const StringType& Str)        { Array.Emplace(MakeShared<FJsonValueString>(FString(Str))); return *this; }
 
-		template <class FNumber>
-		typename TEnableIf<TIsIntegral<FNumber>::Value || TIsFloatingPoint<FNumber>::Value, FArray&>::Type
-			Add(FNumber Number)                   { Array.Emplace(MakeShared<FJsonValueNumber>(Number));   return *this; }
+		template <class NumberType
+			UE_REQUIRES(!std::is_same_v<NumberType, bool> && (std::is_integral_v<NumberType> || std::is_floating_point_v<NumberType>))>
+		FArray& Add(NumberType Number)            { Array.Emplace(MakeShared<FJsonValueNumber>(Number));   return *this; }
 
-		FArray& Add(bool Boolean)                 { Array.Emplace(MakeShared<FJsonValueBoolean>(Boolean)); return *this; }
+		template <class BoolType
+			UE_REQUIRES(std::is_same_v<BoolType, bool>)>
+		FArray& Add(BoolType Boolean)             { Array.Emplace(MakeShared<FJsonValueBoolean>(Boolean)); return *this; }
 		FArray& Add(TYPE_OF_NULLPTR)              { Array.Emplace(MakeShared<FJsonValueNull>());           return *this; }
 
 		FArray& Add(TSharedPtr<FJsonValue> Value) { Array.Emplace(Value);                                  return *this; }
