Json schema functionality.
To be added to JsonUtilties (and new JsonUtilitiesEditor.)

Main developer entry point for schema generation -
.../Source/Runtime/JsonUtilitiles/Public/JsonSchema/JsonSchemaGenerator.h

NOTES -

- There is also now a JsonUtiltiesEditor module, in .../Source/Editor. Json schema generation includes property metadata when run in an Editor context.

- Unit tests are WIP, not part of this CL, and currently being worked elsewhere.

#Synced-CL 48639466

diff --git a/Engine/Source/Editor/JsonUtilitiesEditor/JsonUtilitiesEditor.Build.cs b/Engine/Source/Editor/JsonUtilitiesEditor/JsonUtilitiesEditor.Build.cs
new file mode 100644
--- /dev/null
+++ b/Engine/Source/Editor/JsonUtilitiesEditor/JsonUtilitiesEditor.Build.cs
@@ -0,0 +1,31 @@
+﻿// Copyright Epic Games, Inc. All Rights Reserved.
+
+using UnrealBuildTool;
+
+public class JsonUtilitiesEditor : ModuleRules
+{
+    public JsonUtilitiesEditor(ReadOnlyTargetRules Target) : base(Target)
+    {
+        PCHUsage = ModuleRules.PCHUsageMode.UseExplicitOrSharedPCHs;
+
+        PrivateDependencyModuleNames.AddRange(
+	        new string[]
+	        {
+		        "Core",
+		        "CoreUObject",
+		        "Engine",
+		        "UnrealEd",
+		        "Slate",
+		        "SlateCore",
+		        "Json",
+		        "JsonUtilities",
+		        "BlueprintGraph",
+		        "AssetTools",
+	        }
+        );
+       
+        PrivateIncludePaths.Add(
+	        System.IO.Path.Combine(GetModuleDirectory("JsonUtilities"), "Private")
+	    );
+    }
+}
\ No newline at end of file
diff --git a/Engine/Source/Editor/JsonUtilitiesEditor/Private/JsonSchema/JsonSchemaGeneratorEditor.cpp b/Engine/Source/Editor/JsonUtilitiesEditor/Private/JsonSchema/JsonSchemaGeneratorEditor.cpp
new file mode 100644
--- /dev/null
+++ b/Engine/Source/Editor/JsonUtilitiesEditor/Private/JsonSchema/JsonSchemaGeneratorEditor.cpp
@@ -0,0 +1,206 @@
+// Copyright Epic Games, Inc. All Rights Reserved. 
+
+#include "JsonSchema/JsonSchemaGeneratorEditor.h"
+
+#include "UObject/UnrealType.h"
+#include "EdGraphSchema_K2.h"
+#include "ObjectTools.h"
+
+#include "JsonSchema/JsonSchemaGenerator.h"
+
+//
+// Internal FJsonSchemaMetadataVisitor
+//
+
+class FJsonSchemaMetadataVisitor
+{
+public:
+
+	explicit FJsonSchemaMetadataVisitor(
+			FJsonSchemaEditorMetadata& InEditorMetadata,
+			const FJsonSchemaPropertyFilter& InPropertyFilter = FJsonSchemaPropertyFilter()) :
+		EditorMetadata(InEditorMetadata),
+		PropertyFilter(InPropertyFilter)
+	{
+	}
+	
+	void VisitFProperty(TNotNull<const FProperty*> Property,
+		const bool bIsPropertyFunctionWorldContext = false)
+	{
+		const UFunction* PropertyOwnerFunction = Cast<UFunction>(Property->GetOwnerStruct());
+			
+		// Check to see if we should ignore this property.
+		if (PropertyFilter.IsPropertyIgnored(Property))
+		{
+			return;
+		}
+		
+		const TSharedRef<FJsonObject> PropertyMetadata = MakeShared<FJsonObject>();
+		
+		const FString PropertyName = Property->GetAuthoredName();
+		
+		if (const FText PropertyToolTip = Property->GetToolTipText(); 
+			!PropertyToolTip.IsEmpty())
+		{
+			PropertyMetadata->SetStringField(TEXT("description"), PropertyToolTip.ToString());
+		}
+		else if (PropertyOwnerFunction)
+		{
+			const FString FunctionParamTooltip = (Property->HasAnyPropertyFlags(CPF_ReturnParm) ?
+				ObjectTools::GetDefaultTooltipForFunctionReturn(PropertyOwnerFunction) :
+				ObjectTools::GetDefaultTooltipForFunctionParam(PropertyOwnerFunction, PropertyName));
+			if (!FunctionParamTooltip.IsEmpty())
+			{
+				PropertyMetadata->SetStringField(TEXT("description"), FunctionParamTooltip);
+			}
+		}
+
+		if (const FNumericProperty* NumericProperty = CastField<FNumericProperty>(Property))
+		{
+			if (NumericProperty->HasMetaData("ClampMin"))
+			{
+				PropertyMetadata->SetNumberField(TEXT("minimum"), NumericProperty->GetDoubleMetaData("ClampMin"));
+			}
+			if (NumericProperty->HasMetaData("ClampMax"))
+			{
+				PropertyMetadata->SetNumberField(TEXT("maximum"), NumericProperty->GetDoubleMetaData("ClampMax"));
+			}
+		}
+
+		if (PropertyOwnerFunction)
+		{
+			FString ParameterDefault;
+			UEdGraphSchema_K2::FindFunctionParameterDefaultValue(PropertyOwnerFunction, Property, ParameterDefault);
+			if (!ParameterDefault.IsEmpty())
+			{
+				PropertyMetadata->SetStringField(TEXT("defaultValue"), ParameterDefault);
+			}
+		}
+
+		if (PropertyOwnerFunction && bIsPropertyFunctionWorldContext)
+		{
+			PropertyMetadata->SetBoolField(TEXT("isWorldContext"), true);
+		}
+		
+		if (PropertyMetadata->Values.Num())
+		{
+			EditorMetadata.SetPropertyMetadataForCurrentPropertyMemberPath(PropertyMetadata);
+		}
+
+		// Recurse into inner properties
+		if (const UStruct* InnerStruct = GetFPropertyInnerStruct(*Property, *EditorMetadata.CurrentPropertyMemberPath))
+		{
+			// See all NOTE_JSON_SCHEMA_CHECK_FLAGS_NO_PARMS.
+			TGuardValue CheckFlagsGuard(PropertyFilter.CheckFlags, PropertyFilter.CheckFlags & ~CPF_ParmFlags);
+
+			UE_JSON_SCHEMA_SCOPED_MEMBER_PATH_PUSH(EditorMetadata.CurrentPropertyMemberPath, FJsonSchemaMemberPath::Delimiter, false);
+			
+			VisitUStruct(InnerStruct, false);
+		}
+	}
+	
+	void VisitUStruct(const TNotNull<const UStruct*> Struct, const bool bIsRootStruct)
+	{
+		const UFunction* StructAsFunction = Cast<UFunction>(Struct);
+		
+		const FString FunctionWorldContext =
+			(StructAsFunction && StructAsFunction->HasMetaData(FBlueprintMetadata::MD_WorldContext)) ?
+			StructAsFunction->GetMetaData(FBlueprintMetadata::MD_WorldContext) :
+			"";
+		
+		if (bIsRootStruct)
+		{
+			EditorMetadata.RootStructMetadata = MakeShared<FJsonObject>();
+			
+			if (const FString Description = (
+					StructAsFunction ?
+					ObjectTools::GetDefaultTooltipForFunction(StructAsFunction, false) :
+					Struct->GetToolTipText().ToString());
+				!Description.IsEmpty())
+			{
+				EditorMetadata.RootStructMetadata->SetStringField(TEXT("description"), Description);
+			}
+			
+			if (!FunctionWorldContext.IsEmpty())
+			{
+				EditorMetadata.RootStructMetadata->SetStringField(TEXT("worldContext"), FunctionWorldContext);
+			}
+		}
+		
+		for (const FProperty* Property : TFieldRange<FProperty>(Struct))
+		{
+			const FString PropertyName = Property->GetAuthoredName();
+
+			UE_JSON_SCHEMA_SCOPED_MEMBER_PATH_PUSH(EditorMetadata.CurrentPropertyMemberPath, PropertyName, false);
+			
+			const bool bIsPropertyFunctionWorldContext = (FunctionWorldContext == PropertyName);
+			
+			VisitFProperty(Property,	bIsPropertyFunctionWorldContext);
+		}
+	}
+	
+private:
+
+	UStruct* GetFPropertyInnerStruct(const FProperty& Property, const TCHAR* PropertyPath) const
+	{
+		if (const FArrayProperty* ArrayProperty = CastField<FArrayProperty>(&Property))
+		{
+			if (const FStructProperty* InnerStructProperty = CastField<FStructProperty>(ArrayProperty->Inner))
+			{
+				return InnerStructProperty->Struct;
+			}
+		}
+		else if (const FSetProperty* SetProperty = CastField<FSetProperty>(&Property) )
+		{
+			if (const FStructProperty *ElementStructProperty = CastField<FStructProperty>(SetProperty->ElementProp))
+			{
+				return ElementStructProperty->Struct;
+			}
+		}
+		else if (const FMapProperty* MapProperty = CastField<FMapProperty>(&Property))
+		{
+			if (ensureMsgf(MapProperty->KeyProp->IsA<FStrProperty>(), 
+				TEXT("GetFPropertyInnerStruct() - Unsupported map key type for %s. JSON-compatible map keys must be FString. Map will be skipped."), 
+				PropertyPath))
+			{
+				if (const FStructProperty* ValueStructProperty = CastField<FStructProperty>(MapProperty->ValueProp))
+				{
+					return ValueStructProperty->Struct;
+				}
+			}
+		}
+		else if (const FStructProperty* StructProperty = CastField<FStructProperty>(&Property))
+		{
+			return StructProperty->Struct;
+		}
+		
+		// @todo Push into FObjectProperty via struct instance?
+		return nullptr;
+	}
+
+	FJsonSchemaEditorMetadata& EditorMetadata; 
+	
+	FJsonSchemaPropertyFilter PropertyFilter;
+};
+
+//
+// FJsonSchemaGeneratorEditor
+//
+
+/*static*/ void FJsonSchemaGeneratorEditor::FPropertyToJsonSchemaMetadata(TNotNull<const FProperty*> Property,
+	FJsonSchemaEditorMetadata& OutEditorMetadata,
+	const FJsonSchemaPropertyFilter& PropertyFilter)
+{
+	UE_JSON_SCHEMA_SCOPED_MEMBER_PATH_PUSH(OutEditorMetadata.CurrentPropertyMemberPath, Property->GetAuthoredName(), false);
+
+	FJsonSchemaMetadataVisitor JsonSchemaMetadataVisitor(OutEditorMetadata, PropertyFilter);
+	JsonSchemaMetadataVisitor.VisitFProperty(Property);
+}
+
+/*static*/ void FJsonSchemaGeneratorEditor::UStructToJsonSchemaMetadata(TNotNull<const UStruct*> Struct,
+	FJsonSchemaEditorMetadata& OutEditorMetadata,
+	const FJsonSchemaPropertyFilter& PropertyFilter)
+{
+	FJsonSchemaMetadataVisitor JsonSchemaMetadataVisitor(OutEditorMetadata, PropertyFilter);
+	JsonSchemaMetadataVisitor.VisitUStruct(Struct, true);
+}
diff --git a/Engine/Source/Editor/JsonUtilitiesEditor/Private/JsonUtilitiesEditorModularFeature.cpp b/Engine/Source/Editor/JsonUtilitiesEditor/Private/JsonUtilitiesEditorModularFeature.cpp
new file mode 100644
--- /dev/null
+++ b/Engine/Source/Editor/JsonUtilitiesEditor/Private/JsonUtilitiesEditorModularFeature.cpp
@@ -0,0 +1,27 @@
+// Copyright Epic Games, Inc. All Rights Reserved.
+
+#include "JsonUtilitiesEditorModularFeature.h"
+
+#include "JsonSchema/JsonSchemaGeneratorEditor.h"
+
+PRAGMA_DISABLE_DEPRECATION_WARNINGS
+
+/*virtual*/ void FJsonUtilitiesEditorModularFeature::FPropertyToJsonSchemaMetadata(TNotNull<const FProperty*> Property,
+	FJsonSchemaEditorMetadata& OutEditorMetadata,
+	const FJsonSchemaPropertyFilter& PropertyFilter) const
+{
+	FJsonSchemaGeneratorEditor::FPropertyToJsonSchemaMetadata(Property,
+		OutEditorMetadata,
+		PropertyFilter);
+}
+
+/*virtual*/ void FJsonUtilitiesEditorModularFeature::UStructToJsonSchemaMetadata(TNotNull<const UStruct*> Struct,
+	FJsonSchemaEditorMetadata& OutEditorMetadata,
+	const FJsonSchemaPropertyFilter& PropertyFilter) const
+{
+	return FJsonSchemaGeneratorEditor::UStructToJsonSchemaMetadata(Struct,
+		OutEditorMetadata,
+		PropertyFilter);
+}
+
+PRAGMA_ENABLE_DEPRECATION_WARNINGS
diff --git a/Engine/Source/Editor/JsonUtilitiesEditor/Private/JsonUtilitiesEditorModularFeature.h b/Engine/Source/Editor/JsonUtilitiesEditor/Private/JsonUtilitiesEditorModularFeature.h
new file mode 100644
--- /dev/null
+++ b/Engine/Source/Editor/JsonUtilitiesEditor/Private/JsonUtilitiesEditorModularFeature.h
@@ -0,0 +1,22 @@
+// Copyright Epic Games, Inc. All Rights Reserved.
+
+#pragma once
+
+#include "JsonUtilitiesModularFeature.h"
+
+/**
+ * Derived class for Modular Feature - Editor functionality that the Runtime module will call when in Editor, to collect metadata.
+ * Registered in Editor module.
+ */
+class FJsonUtilitiesEditorModularFeature : public IJsonUtilitiesModularFeature
+{
+public:
+
+	JSONUTILITIESEDITOR_API virtual void FPropertyToJsonSchemaMetadata(TNotNull<const FProperty*> Property,
+		FJsonSchemaEditorMetadata& OutEditorMetadata,
+		const FJsonSchemaPropertyFilter& PropertyFilter = FJsonSchemaPropertyFilter()) const override;
+
+	JSONUTILITIESEDITOR_API virtual void UStructToJsonSchemaMetadata(TNotNull<const UStruct*>,
+		FJsonSchemaEditorMetadata& OutEditorMetadata,
+		const FJsonSchemaPropertyFilter& PropertyFilter = FJsonSchemaPropertyFilter()) const override;
+};
diff --git a/Engine/Source/Editor/JsonUtilitiesEditor/Private/JsonUtilitiesEditorModule.cpp b/Engine/Source/Editor/JsonUtilitiesEditor/Private/JsonUtilitiesEditorModule.cpp
new file mode 100644
--- /dev/null
+++ b/Engine/Source/Editor/JsonUtilitiesEditor/Private/JsonUtilitiesEditorModule.cpp
@@ -0,0 +1,32 @@
+﻿// Copyright Epic Games, Inc. All Rights Reserved.
+
+#include "CoreMinimal.h"
+#include "Modules/ModuleManager.h"
+
+#include "JsonUtilitiesEditorModularFeature.h"
+
+class FJsonUtilitiesEditorModule : public IModuleInterface
+{
+public:
+
+	TUniquePtr<FJsonUtilitiesEditorModularFeature> ModularFeature;
+	
+	virtual void StartupModule() override
+	{
+		FModuleManager::LoadModuleChecked<IModuleInterface>(TEXT("JsonUtilities"));
+		
+		ModularFeature = MakeUnique<FJsonUtilitiesEditorModularFeature>();
+		IModularFeatures::Get().RegisterModularFeature(IJsonUtilitiesModularFeature::Name, ModularFeature.Get());
+	}
+	
+	virtual void ShutdownModule() override
+	{
+		if (ModularFeature.IsValid())
+		{
+			IModularFeatures::Get().UnregisterModularFeature(IJsonUtilitiesModularFeature::Name, ModularFeature.Get());
+			ModularFeature.Reset();
+		}
+	}
+};
+
+IMPLEMENT_MODULE(FJsonUtilitiesEditorModule, JsonUtilitiesEditor)
diff --git a/Engine/Source/Editor/JsonUtilitiesEditor/Public/JsonSchema/JsonSchemaGeneratorEditor.h b/Engine/Source/Editor/JsonUtilitiesEditor/Public/JsonSchema/JsonSchemaGeneratorEditor.h
new file mode 100644
--- /dev/null
+++ b/Engine/Source/Editor/JsonUtilitiesEditor/Public/JsonSchema/JsonSchemaGeneratorEditor.h
@@ -0,0 +1,84 @@
+// Copyright Epic Games, Inc. All Rights Reserved. 
+
+#pragma once
+
+#include "Dom/JsonObject.h"
+#include "Templates/SharedPointer.h"
+#include "UObject/Class.h"
+
+#include "JsonSchema/JsonSchemaPropertyFilter.h"
+#include "JsonSchema/JsonSchemaEditorMetadata.h"
+
+class [[deprecated("Experimental class - Subject to change")]] FJsonSchemaGeneratorEditor
+{
+public:
+	
+	/**
+	 * Usually not called by a developer directly. Used by the Runtime module via the Module Feature.
+	 * 
+	 * @param Property The property to collect metadata for.
+	 * @param OutEditorMetadata The collected property metadata.
+	 * @param PropertyFilter Determines how to filter collected struct and property data.
+	 */
+	JSONUTILITIESEDITOR_API static void FPropertyToJsonSchemaMetadata(TNotNull<const FProperty*> Property,
+		FJsonSchemaEditorMetadata& OutEditorMetadata,
+		const FJsonSchemaPropertyFilter& PropertyFilter = FJsonSchemaPropertyFilter());
+
+	/**
+	 * Usually not called by a developer directly. Used by the Runtime module via the Module Feature.
+	 * 
+	 * Recursively walks Struct's property hierarchy, collecting editor-only meta-data like
+	 * tooltips and other Meta values, caching them in the cookable OutEditorMetadata.
+	 *
+	 * e.g:
+	 *
+	 *		// Details for a person
+	 *		USTRUCT()
+	 *		struct FPerson
+	 *		{
+	 *			GENERATED_BODY()
+	 *	
+	 *			// Full name
+	 *			UPROPERTY()
+	 *			FString Name;
+	 *
+	 *			// How old they are
+	 *			UPROPERTY(meta = (ClampMin = "1"))
+	 *			int32 Age;
+	 *		}
+	 *
+	 *		// A family of people
+	 *		USTRUCT()
+	 *		struct FFamily
+	 *		{
+	 *			// Members of this family
+	 *			UPROPERTY()
+	 *			TArray<FPerson> Members; 
+	 *		}
+	 *
+	 *		-> metadata:
+	 *			description -> "A family of people"
+	 *			PropertyMetadata ->
+	 *				Members = "Members of this family"
+	 *				Members.Name = "Full name"
+	 *				Members.Age = "How old they are" (ClampMin = 1)
+	 *
+	 * Note: When Struct is a UFunction, function parameter descriptions will be harvested from doxygen markup in the function comment e.g:
+	 *
+	 *		// Returns the best friend of the given person
+	 *		// Person: The person to find the best friend of
+	 * 		UFUNCTION()
+	 * 		FPerson GetBestFriend(const FPerson& Person)
+	 *
+	 * 		-> metadata:
+	 *			description -> "Returns the best friend of the given person"
+	 *			PropertyMetadata -> Person -> "The person to find the best friend of"
+	 *
+	 * @param Struct The type to collect property data for. Can be either UScriptStruct, UClass or UFunction.
+	 * @param OutEditorMetadata The collected property metadata.
+	 * @param PropertyFilter Determines how to filter collected struct and property data.
+	 */
+	JSONUTILITIESEDITOR_API static void UStructToJsonSchemaMetadata(TNotNull<const UStruct*> Struct,
+		FJsonSchemaEditorMetadata& OutEditorMetadata,
+		const FJsonSchemaPropertyFilter& PropertyFilter = FJsonSchemaPropertyFilter());
+};
diff --git a/Engine/Source/Runtime/JsonUtilities/Private/JsonSchema/JsonSchemaEditorMetadata.cpp b/Engine/Source/Runtime/JsonUtilities/Private/JsonSchema/JsonSchemaEditorMetadata.cpp
new file mode 100644
--- /dev/null
+++ b/Engine/Source/Runtime/JsonUtilities/Private/JsonSchema/JsonSchemaEditorMetadata.cpp
@@ -0,0 +1,54 @@
+// Copyright Epic Games, Inc. All Rights Reserved. 
+
+#include "JsonSchema/JsonSchemaEditorMetadata.h"
+
+FJsonSchemaEditorMetadata::FJsonSchemaEditorMetadata()
+{
+}
+
+FJsonSchemaEditorMetadata::FJsonSchemaEditorMetadata(const FJsonSchemaEditorMetadata& Other)
+{
+	*this = Other;
+}
+
+FJsonSchemaEditorMetadata& FJsonSchemaEditorMetadata::operator=(const FJsonSchemaEditorMetadata& Other)
+{
+	if (this != &Other)
+	{
+		RootStructMetadata = Other.RootStructMetadata;
+
+		PropertyMemberPathToPropertyMetadataMap = Other.PropertyMemberPathToPropertyMetadataMap;
+
+		CurrentPropertyMemberPath.Reset();
+		CurrentPropertyMemberPath.Append(*const_cast<FJsonSchemaMemberPath&>(Other.CurrentPropertyMemberPath));
+	}
+	
+	return *this;
+}
+
+void FJsonSchemaEditorMetadata::SetPropertyMetadataForCurrentPropertyMemberPath(const TSharedRef<FJsonObject>& PropertyMetadata)
+{
+	PropertyMemberPathToPropertyMetadataMap.Add(*CurrentPropertyMemberPath, PropertyMetadata.ToSharedPtr());
+}
+
+TSharedPtr<FJsonObject> FJsonSchemaEditorMetadata::GetPropertyMetadataForCurrentPropertyMemberPath()
+{
+	if (TSharedPtr<FJsonObject>* Found = PropertyMemberPathToPropertyMetadataMap.Find(*CurrentPropertyMemberPath))
+	{
+		return *Found;
+	}
+	else
+	{
+		return nullptr;
+	}
+}
+
+const FJsonSchemaEditorMetadata::FJsonSchemaPropertyMemberPathToPropertyMetadataMap& FJsonSchemaEditorMetadata::GetPropertyMemberPathToPropertyMetadataMap() const
+{
+	return PropertyMemberPathToPropertyMetadataMap;
+}
+
+FJsonSchemaEditorMetadata::FJsonSchemaPropertyMemberPathToPropertyMetadataMap& FJsonSchemaEditorMetadata::GetPropertyMemberPathToPropertyMetadataMap()
+{
+	return PropertyMemberPathToPropertyMetadataMap;
+}
diff --git a/Engine/Source/Runtime/JsonUtilities/Private/JsonSchema/JsonSchemaGenerator.cpp b/Engine/Source/Runtime/JsonUtilities/Private/JsonSchema/JsonSchemaGenerator.cpp
new file mode 100644
--- /dev/null
+++ b/Engine/Source/Runtime/JsonUtilities/Private/JsonSchema/JsonSchemaGenerator.cpp
@@ -0,0 +1,478 @@
+ // Copyright Epic Games, Inc. All Rights Reserved. 
+
+#include "JsonSchema/JsonSchemaGenerator.h"
+
+#include "JsonDomBuilder.h"
+#include "JsonObjectConverter.h"
+#include "UObject/TextProperty.h"
+#include "UObject/UnrealType.h"
+#include "Policies/CondensedJsonPrintPolicy.h"
+
+#include "JsonUtilitiesModularFeature.h"
+
+// See all NOTE_JSON_SCHEMA_CHECK_FLAGS_NO_PARMS.
+// A UStruct may be a UFunction, in which case a caller might have set the filter flags to ignore
+// certain types of parameters (properties) of the function. After using the CheckFlags to check
+// of a property/parameter should be ignored, we should stop considering flags related to function
+// parameters (properties), by clearing all CPF_ParmFlags, so that they don't block schema
+// generation (or metadata collection) for sub-structs/sub-properties.
+
+//
+// Internal FJsonSchemaGeneratorVisitor
+//
+
+class FJsonSchemaGeneratorVisitor
+{
+public:
+
+	explicit FJsonSchemaGeneratorVisitor(
+			FJsonSchemaEditorMetadata& InEditorMetadata,
+			const FJsonSchemaPropertyFilter& InPropertyFilter = FJsonSchemaPropertyFilter()) :
+		EditorMetadata(InEditorMetadata),
+		PropertyFilter(InPropertyFilter)
+	{
+	}
+	
+	TSharedPtr<FJsonObject> VisitFProperty(TNotNull<const FProperty*> Property)
+	{
+		FJsonDomBuilder::FObject PropertyJson;
+		
+		if (const FOptionalProperty* OptionalProperty = CastField<FOptionalProperty>(Property))
+		{
+			BuildJsonForOptionalProperty(*OptionalProperty, PropertyJson);
+		}
+		else if (const FStructProperty* StructProperty = CastField<FStructProperty>(Property))
+		{
+			BuildJsonForStructProperty(*StructProperty, PropertyJson);
+		}
+		else if (const FArrayProperty* ArrayProperty = CastField<FArrayProperty>(Property))
+		{
+			BuildJsonForArrayProperty(*ArrayProperty, PropertyJson);
+		}
+		else if (const FSetProperty* SetProperty = CastField<FSetProperty>(Property))
+		{
+			BuildJsonForSetProperty(*SetProperty, PropertyJson);
+		}
+		else if (const FMapProperty* MapProperty = CastField<FMapProperty>(Property))
+		{
+			BuildJsonForMapProperty(*MapProperty, PropertyJson);
+		}
+		else if (const FEnumProperty* EnumProperty = CastField<FEnumProperty>(Property))
+		{
+			if (EnumProperty->GetEnum())
+			{
+				BuildJsonForEnumDefinition(*EnumProperty->GetEnum(), PropertyJson);
+			}
+		}
+		else if (const FNumericProperty* NumericProperty = CastField<FNumericProperty>(Property))
+		{
+			BuildJsonForNumericProperty(*NumericProperty, PropertyJson);
+		}
+		else if (CastField<FBoolProperty>(Property))
+		{
+			BuildJsonForBoolProperty(PropertyJson);
+		}
+		else if (CastField<FStrProperty>(Property) || CastField<FTextProperty>(Property) || CastField<FNameProperty>(Property))
+		{
+			BuildJsonForStringProperty(PropertyJson);
+		}
+		else
+		{
+			// Fall back to string for unhandled property types.
+			BuildJsonForStringProperty(PropertyJson);
+		}
+
+		// Add all property metadata, if any.
+		if (const TSharedPtr<FJsonObject> PropertyMetadata = EditorMetadata.GetPropertyMetadataForCurrentPropertyMemberPath())
+		{
+			for (const TTuple<FString, TSharedPtr<FJsonValue>>& Tuple : PropertyMetadata->Values)
+			{
+				const FString JsonMemberName = PropertyFilter.PropertyAuthoredNameToJsonKey(Tuple.Key);
+				PropertyJson.Set(JsonMemberName, Tuple.Value);
+			}
+		}
+		
+		return PropertyJson.AsJsonObject();
+	}
+	
+	TSharedPtr<FJsonObject> VisitUStruct(TNotNull<const UStruct*> Struct, const bool bIsRootStruct)
+	{
+		FJsonDomBuilder::FObject OutStructJson;
+
+		OutStructJson.Set(TEXT("type"), TEXT("object"));
+
+		// Special case handling for FJsonObjectWrapper, where we know this must be a Json object.   
+		if (Struct == FJsonObjectWrapper::StaticStruct())
+		{
+			return OutStructJson.AsJsonObject();
+		}
+
+		// Add all struct metadata.
+		if (bIsRootStruct && EditorMetadata.RootStructMetadata)
+		{
+			for (const TTuple<FString, TSharedPtr<FJsonValue>>& Tuple : EditorMetadata.RootStructMetadata->Values)
+			{
+				const FString JsonMemberName = PropertyFilter.PropertyAuthoredNameToJsonKey(Tuple.Key);
+				OutStructJson.Set(JsonMemberName, Tuple.Value);
+			}
+		}
+
+		// Adds property to JSON and possibly the require list. Returns whether the property was added to the JSON object.
+		auto MaybeAddPropertyToJson = [this](const FProperty* Property,
+			FJsonDomBuilder::FObject& PropertiesJson, FJsonDomBuilder::FArray& RequiredPropertiesArray) -> bool
+			{
+				const FString PropertyName = Property->GetAuthoredName();
+			
+				UE_JSON_SCHEMA_SCOPED_MEMBER_PATH_PUSH(EditorMetadata.CurrentPropertyMemberPath, PropertyName, false);
+
+				if (PropertyFilter.IsPropertyIgnored(Property) || 
+					PropertyFilter.IsPropertyMemberPathSkipped(*EditorMetadata.CurrentPropertyMemberPath))
+				{
+					return false;
+				}
+
+				// See all NOTE_JSON_SCHEMA_CHECK_FLAGS_NO_PARMS.
+				TGuardValue CheckFlagsGuard(PropertyFilter.CheckFlags, PropertyFilter.CheckFlags & ~CPF_ParmFlags);
+			
+				if (const TSharedPtr<FJsonObject> PropertySchema = VisitFProperty(Property);
+					PropertySchema && PropertySchema->Values.Num())
+				{
+					const FString JsonMemberName = PropertyFilter.PropertyAuthoredNameToJsonKey(PropertyName); 
+
+					PropertiesJson.Set(JsonMemberName, PropertySchema);
+
+					const TSharedPtr<FJsonObject> PropertyMetadata = EditorMetadata.GetPropertyMetadataForCurrentPropertyMemberPath();
+
+					if (PropertyFilter.IsPropertyMemberPathRequired(*EditorMetadata.CurrentPropertyMemberPath))
+					{
+						// Property has been explicitly required.
+						RequiredPropertiesArray.Add(JsonMemberName);
+					}
+					else if (!(PropertyMetadata && PropertyMetadata->HasField(TEXT("defaultValue"))) &&
+						!Property->IsA<FOptionalProperty>())
+					{
+						// Property has not been explicitly required.
+						// Only require non-optional properties that don't have default values.
+						RequiredPropertiesArray.Add(JsonMemberName);
+					}
+					
+					return true;
+				}
+				
+				return false;
+			};
+		
+		// Properties & required.
+		// If this struct is a function, use JSON RPC format (with inputSchema/outputSchema). Otherwise, use standard format.
+		if (Struct->IsA(UFunction::StaticClass()))
+		{
+			// If struct is a function, output JSON RPC -
+			// [REDACTED_PROFILE_URL]/modelcontextprotocol/blob/main/schema/2025-06-18/schema.ts#L928
+			// Extension for 'name' -
+			// [REDACTED_PROFILE_URL]/modelcontextprotocol/blob/main/schema/2025-06-18/schema.ts#L312
+			
+			// NOTE - We know we are a UFunction, so will have a valid Outer.
+			OutStructJson.Set(TEXT("name"), *(Struct->GetOuter()->GetName() + "." + Struct->GetAuthoredName()));
+			
+			// Collect input and output properties.
+			FJsonDomBuilder::FObject InputSchemaPropertiesJson, OutputSchemaPropertiesJson;
+			FJsonDomBuilder::FArray InputSchemaRequiredPropertiesArray, OutputSchemaRequiredPropertiesArray;
+			for (const FProperty* Property : TFieldRange<FProperty>(Struct))
+			{
+				if (Property->HasAnyPropertyFlags(CPF_ReturnParm))
+				{
+					MaybeAddPropertyToJson(Property, OutputSchemaPropertiesJson, OutputSchemaRequiredPropertiesArray);
+				}
+				else
+				{
+					MaybeAddPropertyToJson(Property, InputSchemaPropertiesJson, InputSchemaRequiredPropertiesArray);
+				}
+			}
+	
+			// Input schema is not optional.
+			{
+				FJsonDomBuilder::FObject InputSchemaJson;
+				InputSchemaJson.Set(TEXT("type"), TEXT("object"));
+				OutStructJson.Set(TEXT("inputSchema"), InputSchemaJson);
+				
+				if (InputSchemaPropertiesJson.Num() > 0)
+				{
+					InputSchemaJson.Set(TEXT("properties"), InputSchemaPropertiesJson);
+
+					if (InputSchemaRequiredPropertiesArray.Num() > 0)
+					{
+						InputSchemaJson.Set(TEXT("required"), InputSchemaRequiredPropertiesArray);
+					}
+				}
+			}
+	
+			// Output schema is optional.
+			if (OutputSchemaPropertiesJson.Num() > 0)
+			{
+				FJsonDomBuilder::FObject OutputSchemaJson;
+				OutputSchemaJson.Set(TEXT("type"), TEXT("object"));
+				OutStructJson.Set(TEXT("outputSchema"), OutputSchemaJson);
+				
+				// We have OutputSchemaPropertiesJson (see above)
+				{
+					OutputSchemaJson.Set(TEXT("properties"), OutputSchemaPropertiesJson);
+					
+					if (OutputSchemaRequiredPropertiesArray.Num() > 0)
+					{
+						OutputSchemaJson.Set(TEXT("required"), OutputSchemaRequiredPropertiesArray);
+					}
+				}
+			}
+		}
+		else
+		{
+			// If struct is not a function.
+				
+			bool bDidAddProperties = false;
+			FJsonDomBuilder::FObject PropertiesJson;
+			FJsonDomBuilder::FArray RequiredPropertiesArray;
+			for (const FProperty* Property : TFieldRange<FProperty>(Struct))
+			{
+				bDidAddProperties |= MaybeAddPropertyToJson(Property, PropertiesJson, RequiredPropertiesArray);
+			}
+	
+			if (bDidAddProperties)
+			{
+				OutStructJson.Set(TEXT("properties"), PropertiesJson);
+				if (RequiredPropertiesArray.Num() > 0)
+				{
+					OutStructJson.Set(TEXT("required"), RequiredPropertiesArray);
+				}
+			}
+		}
+
+		return OutStructJson.AsJsonObject();
+	}
+	
+private:
+	
+	static void BuildJsonForEnumDefinition(const UEnum& EnumDef,
+		FJsonDomBuilder::FObject& OutJson)
+	{
+		FJsonDomBuilder::FArray EnumNames;
+		const int32 NumEnum = EnumDef.NumEnums();
+		for (int32 EnumIndex = 0; EnumIndex < NumEnum; ++EnumIndex)
+		{
+			const FString StringValue = EnumDef.GetAuthoredNameStringByIndex(EnumIndex);
+			if (StringValue.EndsWith(TEXT("_MAX"), ESearchCase::CaseSensitive))
+			{
+				continue;
+			}
+			EnumNames.Add(StringValue);
+		}
+
+		OutJson.Set(TEXT("type"), TEXT("string"));
+		OutJson.Set(TEXT("enum"), EnumNames.AsJsonValue());
+	}
+
+	void BuildJsonForOptionalProperty(const FOptionalProperty& OptionalProperty,
+		FJsonDomBuilder::FObject& OutJson)
+	{
+		OutJson = VisitFProperty(OptionalProperty.GetValueProperty());
+	}
+
+	void BuildJsonForStructProperty(const FStructProperty& StructProperty,
+		FJsonDomBuilder::FObject& OutJson)
+	{
+		UE_JSON_SCHEMA_SCOPED_MEMBER_PATH_PUSH(EditorMetadata.CurrentPropertyMemberPath, FJsonSchemaMemberPath::Delimiter, false);
+
+		OutJson = VisitUStruct(StructProperty.Struct, false);
+	}
+
+	void BuildJsonForArrayProperty(const FArrayProperty& ArrayProperty,
+		FJsonDomBuilder::FObject& OutJson)
+	{
+		const TSharedPtr<FJsonObject> ItemsJson = VisitFProperty(ArrayProperty.Inner);
+		
+		// Don't double-up the description for the inner items. Especially as the description will be written for the container e.g: "A list of ..."
+		ItemsJson->RemoveField(TEXT("description"));
+		
+		OutJson.Set(TEXT("type"), TEXT("array"));
+		OutJson.Set(TEXT("items"), ItemsJson);
+	}
+
+	void BuildJsonForSetProperty(const FSetProperty& SetProperty,
+		FJsonDomBuilder::FObject& OutJson)
+	{
+		const TSharedPtr<FJsonObject> ItemsJson = VisitFProperty(SetProperty.ElementProp);
+		
+		// Don't double-up the description for the inner items. Especially as the description will be written for the container e.g: "A list of ..."
+		ItemsJson->RemoveField(TEXT("description"));
+
+		OutJson.Set(TEXT("type"), TEXT("array"));
+		OutJson.Set(TEXT("uniqueItems"), true);
+		OutJson.Set(TEXT("items"), ItemsJson);
+	}
+
+	void BuildJsonForMapProperty(const FMapProperty& MapProperty,
+		FJsonDomBuilder::FObject& OutJson)
+	{
+		if (ensureMsgf(MapProperty.KeyProp->IsA<FStrProperty>(), 
+			TEXT("BuildJsonForMapProperty() - Unsupported map key type for %s. JSON-compatible map keys must be FString. Map will be skipped."), 
+			*EditorMetadata.CurrentPropertyMemberPath))
+		{
+			const TSharedPtr<FJsonObject> MapValuesJson = VisitFProperty(MapProperty.ValueProp);
+
+			// Don't double-up the description for the map values. Especially as the description will be written for the container e.g: "A map of ..."
+			MapValuesJson->RemoveField(TEXT("description"));
+
+			OutJson.Set(TEXT("type"), TEXT("object"));
+			OutJson.Set(TEXT("additionalProperties"), MapValuesJson);
+		}
+	}
+
+	static void BuildJsonForNumericProperty(const FNumericProperty& NumericProperty,
+		FJsonDomBuilder::FObject& OutJson)
+	{
+		if (const UEnum* EnumDef = NumericProperty.GetIntPropertyEnum())
+		{
+			BuildJsonForEnumDefinition(*EnumDef, OutJson);
+		}
+		else if (NumericProperty.IsFloatingPoint())
+		{
+			OutJson.Set(TEXT("type"), TEXT("number"));
+		}
+		else if (NumericProperty.IsInteger())
+		{
+			OutJson.Set(TEXT("type"), TEXT("integer"));
+		}
+	}
+
+	static void BuildJsonForBoolProperty(FJsonDomBuilder::FObject& OutJson)
+	{
+		OutJson.Set(TEXT("type"), TEXT("boolean"));
+	}
+
+	static void BuildJsonForStringProperty(FJsonDomBuilder::FObject& OutJson)
+	{
+		OutJson.Set(TEXT("type"), TEXT("string"));
+	}
+
+	FJsonSchemaEditorMetadata& EditorMetadata; 
+	
+	FJsonSchemaPropertyFilter PropertyFilter;
+};
+
+//
+// Internal
+//
+
+namespace 
+{
+	template<class CharType, class PrintPolicy>
+	FString JsonObjectToJsonObjectString(const TSharedRef<FJsonObject>& JsonObject, const int32 Indent)
+	{
+		FString OutString;
+		
+		TSharedRef<TJsonWriter<CharType, PrintPolicy> > JsonWriter = TJsonWriterFactory<CharType, PrintPolicy>::Create(&OutString, Indent);
+		const bool bSuccess = FJsonSerializer::Serialize(JsonObject, JsonWriter);
+		JsonWriter->Close();
+		
+		return (bSuccess ? OutString : TEXT(""));
+	}
+}
+
+//
+// FJsonSchemaGenerator
+//
+// See all NOTE_JSON_UTILITIES_SCHEMA_GENERATOR_METADATA_MODULAR_FEATURE.
+// If a Modular Feature is registered, for collecting struct metadata, call it to get
+// metadata. Collecting metadata from structs and their properties can happen in the
+// Editor context, where the Modular Feature will be registered. (If not running in
+// Editor context, no metadata can be collected, and the Modular Feature will not be
+// registered.)
+//
+
+/*static*/ TSharedPtr<FJsonObject> FJsonSchemaGenerator::FPropertyToJsonSchemaObject(TNotNull<const FProperty*> Property,
+	const FJsonSchemaPropertyFilter& PropertyFilter, 
+	const FJsonSchemaEditorMetadata* CachedEditorMetadata)
+{
+	// See all NOTE_JSON_UTILITIES_SCHEMA_GENERATOR_METADATA_MODULAR_FEATURE.
+	FJsonSchemaEditorMetadata EditorMetadata;
+	if (CachedEditorMetadata)
+	{
+		EditorMetadata = *const_cast<FJsonSchemaEditorMetadata*>(CachedEditorMetadata);
+	}
+	else if (IJsonUtilitiesModularFeature* ModuleFeature = IJsonUtilitiesModularFeature::Get())
+	{
+		ModuleFeature->FPropertyToJsonSchemaMetadata(Property,
+			EditorMetadata,
+			PropertyFilter);
+	}
+	
+	UE_JSON_SCHEMA_SCOPED_MEMBER_PATH_PUSH(EditorMetadata.CurrentPropertyMemberPath, *Property->GetAuthoredName(), false);
+	
+	FJsonSchemaGeneratorVisitor JsonSchemaGeneratorVisitor(EditorMetadata, PropertyFilter);
+	
+	return JsonSchemaGeneratorVisitor.VisitFProperty(Property);
+}
+
+/*static*/ FString FJsonSchemaGenerator::FPropertyToJsonSchemaString(TNotNull<const FProperty*> Property,
+	const FJsonSchemaPropertyFilter& PropertyFilter,
+	const FJsonSchemaEditorMetadata* CachedEditorMetadata,
+	const int32 Indent, const bool bPrettyPrint)
+{
+	FString OutString;
+	
+	if (const TSharedPtr<FJsonObject> JsonSchema = FPropertyToJsonSchemaObject(Property, PropertyFilter, CachedEditorMetadata))
+	{
+		if (bPrettyPrint)
+		{
+			OutString = JsonObjectToJsonObjectString<TCHAR, TPrettyJsonPrintPolicy<TCHAR>>(JsonSchema.ToSharedRef(), Indent);
+		}
+		else
+		{
+			OutString = JsonObjectToJsonObjectString<TCHAR, TCondensedJsonPrintPolicy<TCHAR>>(JsonSchema.ToSharedRef(), Indent);
+		}
+	}
+
+	return OutString;
+}
+
+/*static*/ TSharedPtr<FJsonObject> FJsonSchemaGenerator::UStructToJsonSchemaObject(TNotNull<const UStruct*> Struct,
+	const FJsonSchemaPropertyFilter& PropertyFilter,
+	const FJsonSchemaEditorMetadata* CachedEditorMetadata)
+{
+	// See all NOTE_JSON_UTILITIES_SCHEMA_GENERATOR_METADATA_MODULAR_FEATURE.
+	FJsonSchemaEditorMetadata EditorMetadata;
+	if (CachedEditorMetadata)
+	{
+		EditorMetadata = *const_cast<FJsonSchemaEditorMetadata*>(CachedEditorMetadata);
+	}
+	else if (IJsonUtilitiesModularFeature* ModuleFeature = IJsonUtilitiesModularFeature::Get())
+	{
+		ModuleFeature->UStructToJsonSchemaMetadata(Struct, 
+			EditorMetadata,
+			PropertyFilter);
+	}
+	
+	FJsonSchemaGeneratorVisitor JsonSchemaGeneratorVisitor(EditorMetadata, PropertyFilter);
+	return JsonSchemaGeneratorVisitor.VisitUStruct(Struct, true);
+}
+
+/*static*/ FString FJsonSchemaGenerator::UStructToJsonSchemaString(TNotNull<const UStruct*> Struct,
+	const FJsonSchemaPropertyFilter& PropertyFilter,
+	const FJsonSchemaEditorMetadata* CachedEditorMetadata,
+	const int32 Indent, const bool bPrettyPrint)
+{
+	FString OutString;
+	
+	if (const TSharedPtr<FJsonObject> JsonSchema = UStructToJsonSchemaObject(Struct, PropertyFilter, CachedEditorMetadata))
+	{
+		if (bPrettyPrint)
+		{
+			OutString = JsonObjectToJsonObjectString<TCHAR, TPrettyJsonPrintPolicy<TCHAR>>(JsonSchema.ToSharedRef(), Indent);
+		}
+		else
+		{
+			OutString = JsonObjectToJsonObjectString<TCHAR, TCondensedJsonPrintPolicy<TCHAR>>(JsonSchema.ToSharedRef(), Indent);
+		}
+	}
+
+	return OutString;
+}
diff --git a/Engine/Source/Runtime/JsonUtilities/Private/JsonSchema/JsonSchemaMemberPath.cpp b/Engine/Source/Runtime/JsonUtilities/Private/JsonSchema/JsonSchemaMemberPath.cpp
new file mode 100644
--- /dev/null
+++ b/Engine/Source/Runtime/JsonUtilities/Private/JsonSchema/JsonSchemaMemberPath.cpp
@@ -0,0 +1,100 @@
+// Copyright Epic Games, Inc. All Rights Reserved. 
+
+#include "JsonSchema/JsonSchemaMemberPath.h"
+
+//
+// FJsonSchemaMemberPath
+//
+
+/*static*/ const FStringView FJsonSchemaMemberPath::Delimiter = TEXTVIEW(".");
+
+FJsonSchemaMemberPath::FJsonSchemaMemberPath()
+{
+}
+
+FJsonSchemaMemberPath::FJsonSchemaMemberPath(const FJsonSchemaMemberPath& InMemberPath)
+{
+	*this = InMemberPath;
+}
+
+FJsonSchemaMemberPath::FJsonSchemaMemberPath(FStringView InitialString)
+{
+	Append(InitialString);
+}
+
+FJsonSchemaMemberPath& FJsonSchemaMemberPath::operator=(const FJsonSchemaMemberPath& Other)
+{
+	if (this != &Other)
+	{
+		Reset();
+		
+		FJsonSchemaMemberPath& InMemberPathMutable = *const_cast<FJsonSchemaMemberPath*>(&Other); 
+		Append(*InMemberPathMutable);
+	}
+	
+	return *this;
+}
+
+bool FJsonSchemaMemberPath::Push(FStringView String)
+{
+	if (String.IsEmpty())
+	{
+		return false;
+	}
+	
+	Append(String);
+
+	return true;
+}
+
+bool FJsonSchemaMemberPath::Pop(FStringView String)
+{
+	if (String.IsEmpty())
+	{
+		return false;
+	}
+	
+	// Safety check. Make sure current property path indeed ends with the text the caller wants to trim from the end.
+	{
+		const int32 Offset = Len() - String.Len();
+		if (!ensureAlways(Offset >= 0))
+		{
+			return false;
+		}
+		
+		if (!ensureAlways(String.Equals(ToString() + Offset, ESearchCase::CaseSensitive)))
+		{
+			return false;
+		}
+	}
+	
+	RemoveSuffix(String.Len());
+	
+	return true;
+}
+
+bool FJsonSchemaMemberPath::PushSubmember(FStringView SubmemberString)
+{
+	if (SubmemberString.IsEmpty())
+	{
+		return false;
+	}
+	
+	return (Push(Delimiter) && Push(SubmemberString));
+}
+
+bool FJsonSchemaMemberPath::PopSubmember(FStringView SubmemberString)
+{
+	if (SubmemberString.IsEmpty())
+	{
+		return false;
+	}
+	
+	return (Pop(SubmemberString) && Pop(Delimiter));
+}
+
+FText FJsonSchemaMemberPath::ToText() const
+{
+	FJsonSchemaMemberPath& ThisMutable = *const_cast<FJsonSchemaMemberPath*>(this); 
+	return FText::FromString(*ThisMutable);
+}
diff --git a/Engine/Source/Runtime/JsonUtilities/Private/JsonSchema/JsonSchemaPropertyFilter.cpp b/Engine/Source/Runtime/JsonUtilities/Private/JsonSchema/JsonSchemaPropertyFilter.cpp
new file mode 100644
--- /dev/null
+++ b/Engine/Source/Runtime/JsonUtilities/Private/JsonSchema/JsonSchemaPropertyFilter.cpp
@@ -0,0 +1,60 @@
+// Copyright Epic Games, Inc. All Rights Reserved. 
+
+#include "JsonSchema/JsonSchemaPropertyFilter.h"
+
+FJsonSchemaPropertyFilter::FJsonSchemaPropertyFilter(
+	const EPropertyFlags InCheckFlags,
+	const EPropertyFlags InSkipFlags,
+	const TOptional<TSet<FString>>& InRequiredPropertyMemberPaths,
+	const TOptional<TSet<FString>>& InSkipPropertyMemberPaths,
+	const EJsonObjectConversionFlags InConversionFlags)
+	:
+	CheckFlags(InCheckFlags),
+	SkipFlags(InSkipFlags),
+	RequiredPropertyMemberPaths(InRequiredPropertyMemberPaths),
+	SkipPropertyMemberPaths(InSkipPropertyMemberPaths),
+	ConversionFlags(InConversionFlags)
+{
+}
+
+FJsonSchemaPropertyFilter::FJsonSchemaPropertyFilter(const FJsonSchemaPropertyFilter& Other)
+{
+	CheckFlags = Other.CheckFlags;
+	SkipFlags = Other.SkipFlags;
+	RequiredPropertyMemberPaths = Other.RequiredPropertyMemberPaths;
+	SkipPropertyMemberPaths = Other.SkipPropertyMemberPaths;
+	ConversionFlags = Other.ConversionFlags;
+}
+
+bool FJsonSchemaPropertyFilter::IsPropertyIgnored(TNotNull<const FProperty*> Property) const
+{
+	return (Property->HasAnyPropertyFlags(SkipFlags) || (CheckFlags != CPF_None && !Property->HasAnyPropertyFlags(CheckFlags)));
+}
+
+bool FJsonSchemaPropertyFilter::IsPropertyMemberPathRequired(const FString& PropertyMemberPath) const
+{
+	if (PropertyMemberPath.IsEmpty())
+	{
+		return false;
+	}
+	
+	return (RequiredPropertyMemberPaths.IsSet() && RequiredPropertyMemberPaths->Contains(PropertyMemberPath));
+}
+
+bool FJsonSchemaPropertyFilter::IsPropertyMemberPathSkipped(const FString& PropertyMemberPath) const
+{
+	if (PropertyMemberPath.IsEmpty())
+	{
+		return false;
+	}
+	
+	return (SkipPropertyMemberPaths.IsSet() && SkipPropertyMemberPaths->Contains(PropertyMemberPath));
+}
+
+FString FJsonSchemaPropertyFilter::PropertyAuthoredNameToJsonKey(const FString& PropertyAuthoredName) const
+{
+	// TODO - FJsonObjectConverter::StandardizeCase() leaves a lot to be desired. We should look into writing a new one.
+	return (EnumHasAnyFlags(ConversionFlags, EJsonObjectConversionFlags::SkipStandardizeCase) ?
+		PropertyAuthoredName :
+		FJsonObjectConverter::StandardizeCase(PropertyAuthoredName));
+}
diff --git a/Engine/Source/Runtime/JsonUtilities/Private/JsonUtilitiesModularFeature.cpp b/Engine/Source/Runtime/JsonUtilities/Private/JsonUtilitiesModularFeature.cpp
new file mode 100644
--- /dev/null
+++ b/Engine/Source/Runtime/JsonUtilities/Private/JsonUtilitiesModularFeature.cpp
@@ -0,0 +1,18 @@
+// Copyright Epic Games, Inc. All Rights Reserved.
+
+#include "JsonUtilitiesModularFeature.h"
+
+/*static*/ const FName IJsonUtilitiesModularFeature::Name(TEXT("JsonUtilitiesModularFeature"));
+
+IJsonUtilitiesModularFeature::IJsonUtilitiesModularFeature()
+{
+}
+
+/*static*/ IJsonUtilitiesModularFeature* IJsonUtilitiesModularFeature::Get()
+{
+	// See all NOTE_JSON_UTILITIES_SCHEMA_GENERATOR_METADATA_MODULAR_FEATURE.
+	IModularFeatures& ModularFeatures = IModularFeatures::Get();
+	return (ModularFeatures.IsModularFeatureAvailable(Name) ?
+		&ModularFeatures.GetModularFeature<IJsonUtilitiesModularFeature>(Name) :
+		nullptr);
+}
diff --git a/Engine/Source/Runtime/JsonUtilities/Private/JsonUtilitiesModularFeature.h b/Engine/Source/Runtime/JsonUtilities/Private/JsonUtilitiesModularFeature.h
new file mode 100644
--- /dev/null
+++ b/Engine/Source/Runtime/JsonUtilities/Private/JsonUtilitiesModularFeature.h
@@ -0,0 +1,34 @@
+// Copyright Epic Games, Inc. All Rights Reserved.
+
+#pragma once
+
+#include "Features/IModularFeature.h"
+#include "Features/IModularFeatures.h"
+
+#include "JsonSchema/JsonSchemaPropertyFilter.h"
+#include "JsonSchema/JsonSchemaEditorMetadata.h"
+
+/**
+ * Base class for Modular Feature - Editor functionality that the Runtime module will call when in Editor, to collect metadata.
+ * See Editor module for derivation and registration.
+ */
+class IJsonUtilitiesModularFeature : public IModularFeature
+{
+public:
+	
+	JSONUTILITIES_API static const FName Name;
+
+	JSONUTILITIES_API IJsonUtilitiesModularFeature();
+	virtual ~IJsonUtilitiesModularFeature() = default;
+	
+	JSONUTILITIES_API static IJsonUtilitiesModularFeature* Get();
+
+	
+	virtual void FPropertyToJsonSchemaMetadata(TNotNull<const FProperty*> Property,
+		FJsonSchemaEditorMetadata& OutEditorMetadata,
+		const FJsonSchemaPropertyFilter& PropertyFilter = FJsonSchemaPropertyFilter()) const = 0;
+
+	virtual void UStructToJsonSchemaMetadata(TNotNull<const UStruct*> Struct,
+		FJsonSchemaEditorMetadata& OutEditorMetadata,
+		const FJsonSchemaPropertyFilter& PropertyFilter = FJsonSchemaPropertyFilter()) const = 0;
+};
diff --git a/Engine/Source/Runtime/JsonUtilities/Public/JsonSchema/JsonSchemaEditorMetadata.h b/Engine/Source/Runtime/JsonUtilities/Public/JsonSchema/JsonSchemaEditorMetadata.h
new file mode 100644
--- /dev/null
+++ b/Engine/Source/Runtime/JsonUtilities/Public/JsonSchema/JsonSchemaEditorMetadata.h
@@ -0,0 +1,59 @@
+// Copyright Epic Games, Inc. All Rights Reserved. 
+
+#pragma once
+
+#include "CoreMinimal.h"
+#include "Dom/JsonObject.h"
+
+#include "JsonSchema/JsonSchemaMemberPath.h"
+
+/**
+ * Used for constructing a member path for properties, and storing/accessing metadata for multiple
+ * property member paths. Also stores metadata for a root struct.
+ */
+struct FJsonSchemaEditorMetadata
+{
+	using FJsonSchemaPropertyMemberPathToPropertyMetadataMap = TMap<FString, TSharedPtr<FJsonObject>>;
+	
+	JSONUTILITIES_API FJsonSchemaEditorMetadata();
+	
+	JSONUTILITIES_API FJsonSchemaEditorMetadata(const FJsonSchemaEditorMetadata& Other);
+	
+	JSONUTILITIES_API FJsonSchemaEditorMetadata& operator=(const FJsonSchemaEditorMetadata& Other);
+
+	/**
+	 * Top-level metadata (if any) for the root struct only. Mainly used when UStruct is UFunction.
+	 */
+	TSharedPtr<FJsonObject> RootStructMetadata;
+		
+	/**
+	 * Adds property metadata for the current property member path.
+	 * @param PropertyMetadata Property metadata.
+	 */
+	JSONUTILITIES_API void SetPropertyMetadataForCurrentPropertyMemberPath(const TSharedRef<FJsonObject>& PropertyMetadata);
+
+	/**
+	 * Gets property metadata for the current property member path.
+	 * @return Property metadata for the current property member path.
+	 */
+	JSONUTILITIES_API TSharedPtr<FJsonObject> GetPropertyMetadataForCurrentPropertyMemberPath();
+
+	/**
+	 * Get Property member path -to- property metadata map.
+	 */
+	JSONUTILITIES_API const FJsonSchemaPropertyMemberPathToPropertyMetadataMap& GetPropertyMemberPathToPropertyMetadataMap() const;
+	JSONUTILITIES_API FJsonSchemaPropertyMemberPathToPropertyMetadataMap& GetPropertyMemberPathToPropertyMetadataMap();
+	
+	/**
+	 * Current property member path being constructed.
+	 */
+	FJsonSchemaMemberPath CurrentPropertyMemberPath;
+
+private:
+
+	/**
+	 * Maps completed property member paths to property metadata, if present.
+	 * Per-property metadata, including inner properties e.g: Person, Person.Name, Person.Age
+	 */
+	FJsonSchemaPropertyMemberPathToPropertyMetadataMap PropertyMemberPathToPropertyMetadataMap;
+};
diff --git a/Engine/Source/Runtime/JsonUtilities/Public/JsonSchema/JsonSchemaGenerator.h b/Engine/Source/Runtime/JsonUtilities/Public/JsonSchema/JsonSchemaGenerator.h
new file mode 100644
--- /dev/null
+++ b/Engine/Source/Runtime/JsonUtilities/Public/JsonSchema/JsonSchemaGenerator.h
@@ -0,0 +1,230 @@
+// Copyright Epic Games, Inc. All Rights Reserved. 
+
+#pragma once
+
+#include "CoreMinimal.h"
+#include "Dom/JsonObject.h"
+#include "Internationalization/Text.h"
+#include "UObject/Class.h"
+#include "Templates/SharedPointer.h"
+
+#include "JsonSchemaPropertyFilter.h"
+#include "JsonSchemaEditorMetadata.h"
+
+/**
+ * Used for generating a JSON schema from any UStruct or FProperty.
+ * 
+ * See all NOTE_JSON_UTILITIES_EXPERIMENTAL_CLASSES.
+ * Currently deprecated to display warnings in build, and warning state in IDEs.
+ */
+class [[deprecated("Experimental class - Subject to change")]] FJsonSchemaGenerator
+{
+public:
+	
+	/**
+	 * Derives a JSON Schema definition for Property's type.
+	 *
+	 * Can be called by a developer directly, but primarily used by UStructToJsonSchemaObject().
+	 * Attention -
+	 *		- Property name is not part of the returned schema.
+	 *		- If calling in Editor context, metadata will be included. Otherwise, it will not. (i.e. default and min/max clamp values.)
+	 *		- If Property's owner Struct is a UFunction and calling from an Editor context, metadata will only include default values if UFunction is
+	 *		  a Blueprint function.
+	 * 
+	 * If Property is a struct or container of structs property etc, the nested types will be
+	 * recursed into using UStructToJsonSchemaObject to derive a fully complete schema.  
+	 *
+	 * Importantly, the derived JSON Schema describes the JSON which would be produced by FJsonObjectConverter::UPropertyToJsonValue for a property 
+	 * value of this type. And conversely, any JSON adhering to this JSON Schema, may then be deserialized to this property,
+	 * using FJsonObjectConverter::JsonValueToUProperty.     
+	 *
+	 * e.g:
+	 *
+	 *		FBoolProperty -> { "type": "boolean" }
+	 *		FStrProperty & FTextProperty -> { "type": "string" }
+	 *		FStructProperty -> { "type": "object", "properties": { ... } }
+	 * 
+	 * @param Property The property to build a schema for.
+	 * @param PropertyFilter Determines how to filter collected struct and property data.
+	 * @param CachedEditorMetadata If valid, use this metadata instead of collecting it. Otherwise, it will be automatically collected
+	 *		(if running in Editor context.) This is optional, and used for workflows that have cached metadata from an earlier process. 
+	 * @return Schema for the FProperty, as a JSON object. 
+	 */
+	JSONUTILITIES_API TSharedPtr<FJsonObject> static FPropertyToJsonSchemaObject(TNotNull<const FProperty*> Property,
+		const FJsonSchemaPropertyFilter& PropertyFilter = FJsonSchemaPropertyFilter(), 
+		const FJsonSchemaEditorMetadata* CachedEditorMetadata = nullptr);
+	
+	/**
+	 * String output variant of FPropertyToJsonSchemaObject() (see above.)
+	 * 
+	 * Provides JSON string instead of JSON object.
+	 * 
+	 * @param Property The property to build a schema for.
+	 * @param PropertyFilter Determines how to filter collected struct and property data.
+	 * @param CachedEditorMetadata If valid, use this metadata instead of collecting it. Otherwise, it will be automatically collected
+	 *		(if running in Editor context.) This is optional, and used for workflows that have cached metadata from an earlier process. 
+	 * @param Indent How many additional tabs to add to the json serializer output.
+	 * @param bPrettyPrint Option to use pretty print (e.g., adds line endings), otherwise condensed print.
+	 * @return Schema for the property, as a JSON string. 
+	 */
+	JSONUTILITIES_API static FString FPropertyToJsonSchemaString(TNotNull<const FProperty*> Property,
+		const FJsonSchemaPropertyFilter& PropertyFilter = FJsonSchemaPropertyFilter(), 
+		const FJsonSchemaEditorMetadata* CachedEditorMetadata = nullptr,
+		const int32 Indent = 0, const bool bPrettyPrint = true);
+	
+	/**
+	 * Recursively walks Struct's property hierarchy to derive a JSON Schema definition for this type.
+	 *
+	 * To be called by a developer directly.
+	 * Attention -
+	 *		- Struct name is not part of the returned schema.
+	 *		- If calling in Editor context, metadata will be included. Otherwise, it will not. (i.e. struct's description)
+	 *		- If Struct is a UFunction, and calling from an Editor context, metadata will only include default values for the function parameters if 
+	 *		  the function is a Blueprint function.
+	 *		- If Struct is a UClass, its UFunctions will not be included in the returned schema. TODO - Need a UClass version of this function? 
+	 * 
+	 * Importantly, the derived JSON Schema describes the JSON which would be produced by FJsonObjectConverter::StructToJsonObject for an instance
+	 * of the given UStruct. And conversely, any JSON adhering to this JSON Schema, may then be deserialized to an instance of Struct,
+	 * using FJsonObjectConverter::JsonObjectToUStruct.     
+	 *
+	 * e.g:
+	 *
+	 *		// Details for a person
+	 *		USTRUCT()
+	 *		struct FPerson
+	 *		{
+	 *			GENERATED_BODY()
+	 *	
+	 *			// Full name
+	 *			UPROPERTY()
+	 *			FString Name;
+	 *
+	 *			// How old they are
+	 *			UPROPERTY(meta = (ClampMin = "1"))
+	 *			int32 Age;
+	 *		}
+	 *
+	 *		// A family of people
+	 *		USTRUCT()
+	 *		struct FFamily
+	 *		{
+	 *			// Members of this family
+	 *			UPROPERTY()
+	 *			TArray<FPerson> Members; 
+	 *		}
+	 *
+	 *		->
+	 *		
+	 *		{
+	 *			"type": "object",
+	 *			"properties": {
+	 *				"members": {
+	 *					"type": "array",
+	 *					"description": "Members of this family",
+	 *					"items": {
+	 *						"type": "object",
+	 *						"properties": {
+	 *							"name": {
+	 *								"type": "string",
+	 *								"description": "Full name"
+	 *							},
+	 *							"age": {
+	 *								"type": "integer",
+	 *								"description": "How old they are",
+	 *								"minimum": 1
+	 *							}
+	 *						},
+	 *						"required": ["name", "age"]
+	 *					}
+	 *				}
+	 *			},
+	 *			"required": ["members"]
+	 *		}
+	 *
+	 * IMPORTANT - If called while in Editor, schema will include metadata. Otherwise, it will not.
+	 * 
+	 * @param Struct The type to build a schema for. Can be either UScriptStruct, UClass or UFunction.
+	 * @param PropertyFilter Determines how to filter collected struct and property data.
+	 * @param CachedEditorMetadata If valid, use this metadata instead of collecting it. Otherwise, it will be automatically collected
+	 *		(if running in Editor context.) This is optional, and used for workflows that have cached metadata from an earlier process. 
+	 * @return Schema for the UStruct, as a JSON object. 
+	 */
+	JSONUTILITIES_API static TSharedPtr<FJsonObject> UStructToJsonSchemaObject(TNotNull<const UStruct*> Struct,
+		const FJsonSchemaPropertyFilter& PropertyFilter = FJsonSchemaPropertyFilter(),
+		const FJsonSchemaEditorMetadata* CachedEditorMetadata = nullptr);
+	
+	/**
+	 * Templated variant of UStructToJsonSchemaObject() (see above.).
+	 * 
+	 * Works according to struct/class type instead of an taking actual UStruct/UClass reference -
+	 * i.e. UStructToJsonSchema<FSomeStruct>() or UStructToJsonSchema<USomeClass>().
+	 * 
+	 * @param PropertyFilter Determines how to filter collected struct and property data.
+	 * @param CachedEditorMetadata If valid, use this metadata instead of collecting it. Otherwise, it will be automatically collected
+	 *		(if running in Editor context.) This is optional, and used for workflows that have cached metadata from an earlier process. 
+	 * @return Schema for the UStruct, as a JSON object. 
+	 */
+	template<typename InStructType>
+	static TSharedPtr<FJsonObject> UStructToJsonSchemaObject(
+		const FJsonSchemaPropertyFilter& PropertyFilter = FJsonSchemaPropertyFilter(),
+		FJsonSchemaEditorMetadata* CachedEditorMetadata = nullptr)
+	{
+		// Compile-time discovery if struct type has StaticClass(), so is UClass.
+		if constexpr (TModels<CStaticClassProvider, InStructType>::Value)
+		{
+			return UStructToJsonSchemaObject(InStructType::StaticClass(), PropertyFilter, CachedEditorMetadata);
+		}
+		else
+		{
+			return UStructToJsonSchemaObject(InStructType::StaticStruct(), PropertyFilter, CachedEditorMetadata);
+		}
+	}
+
+	/**
+	 * String output variant of UStructToJsonSchemaObject() (see above.)
+	 * 
+	 * Provides JSON string instead of JSON object.
+	 *
+	 * @param Struct The type to build a schema for. Can be either UScriptStruct, UClass or UFunction.
+	 * @param PropertyFilter Determines how to filter collected struct and property data.
+	 * @param CachedEditorMetadata If valid, use this metadata instead of collecting it. Otherwise, it will be automatically collected
+	 *		(if running in Editor context.) This is optional, and used for workflows that have cached metadata from an earlier process. 
+	 * @param Indent How many additional tabs to add to the json serializer output.
+	 * @param bPrettyPrint Option to use pretty print (e.g., adds line endings), otherwise condensed print.
+	 * @return Schema for the UStruct, as a JSON string. 
+	 */
+	JSONUTILITIES_API static FString UStructToJsonSchemaString(TNotNull<const UStruct*> Struct,
+		const FJsonSchemaPropertyFilter& PropertyFilter = FJsonSchemaPropertyFilter(),
+		const FJsonSchemaEditorMetadata* CachedEditorMetadata = nullptr,
+		const int32 Indent = 0, const bool bPrettyPrint = true);
+	
+	/**
+	 * Templated variant of UStructToJsonSchemaString() (see above.).
+	 * 
+	 * Works according to struct/class type instead of an taking actual UStruct/UClass reference -
+	 * i.e. UStructToJsonSchema<FSomeStruct>() or UStructToJsonSchema<USomeClass>().
+	 * 
+	 * @param PropertyFilter Determines how to filter collected struct and property data.
+	 * @param CachedEditorMetadata If valid, use this metadata instead of collecting it. Otherwise, it will be automatically collected
+	 *		(if running in Editor context.) This is optional, and used for workflows that have cached metadata from an earlier process. 
+	 * @param Indent How many additional tabs to add to the json serializer output.
+	 * @param bPrettyPrint Option to use pretty print (e.g., adds line endings), otherwise condensed print.
+	 * @return Schema for the UStruct, as a JSON string. 
+	 */
+	template<typename InStructType>
+	static FString UStructToJsonSchemaString(
+		const FJsonSchemaPropertyFilter& PropertyFilter = FJsonSchemaPropertyFilter(),
+		const FJsonSchemaEditorMetadata* CachedEditorMetadata = nullptr,
+		const int32 Indent = 0, const bool bPrettyPrint = true)
+	{
+		// Compile-time discovery if struct type has StaticClass(), so is UClass.
+		if constexpr (TModels<CStaticClassProvider, InStructType>::Value)
+		{
+			return UStructToJsonSchemaString(*InStructType::StaticClass(), PropertyFilter, CachedEditorMetadata, Indent, bPrettyPrint);
+		}
+		else
+		{
+			return UStructToJsonSchemaString(*InStructType::StaticStruct(), PropertyFilter, CachedEditorMetadata, Indent, bPrettyPrint);
+		}
+	}
+};
diff --git a/Engine/Source/Runtime/JsonUtilities/Public/JsonSchema/JsonSchemaMemberPath.h b/Engine/Source/Runtime/JsonUtilities/Public/JsonSchema/JsonSchemaMemberPath.h
new file mode 100644
--- /dev/null
+++ b/Engine/Source/Runtime/JsonUtilities/Public/JsonSchema/JsonSchemaMemberPath.h
@@ -0,0 +1,105 @@
+// Copyright Epic Games, Inc. All Rights Reserved. 
+
+#pragma once
+
+#include "CoreMinimal.h"
+
+/**
+ * JSON member path builder.
+ */
+class FJsonSchemaMemberPath : public TStringBuilder<512>
+{
+public:
+
+	/**
+	 * Text used to represent divider between member names and sub-member names.
+	 */
+	JSONUTILITIES_API static const FStringView Delimiter;
+
+	JSONUTILITIES_API FJsonSchemaMemberPath();
+	JSONUTILITIES_API FJsonSchemaMemberPath(const FJsonSchemaMemberPath& InMemberPath);
+	JSONUTILITIES_API explicit FJsonSchemaMemberPath(FStringView InitialString);
+	
+	FJsonSchemaMemberPath& operator=(const FJsonSchemaMemberPath& Other);
+	
+	/**
+	 * Add a string to the end.
+	 * @param String String add to end.
+	 * @return Success.
+	 */
+	JSONUTILITIES_API bool Push(FStringView String);
+	
+	/**
+	 * Remove a string from the end.
+	 * @param String String to remove from end.
+	 * @return Success.
+	 */
+	JSONUTILITIES_API bool Pop(FStringView String);
+
+	/**
+	 * Add a string to the end, as a submember: Delimiter + SubmemberText
+	 * @param SubmemberString String add to end.
+	 * @return Success.
+	 */
+	JSONUTILITIES_API bool PushSubmember(FStringView SubmemberString);
+	
+	/**
+	 * Remove a string from the end, as a submember: Delimiter + SubmemberText
+	 * @param SubmemberString String to remove from end.
+	 * @return Success.
+	 */
+	JSONUTILITIES_API bool PopSubmember(FStringView SubmemberString);
+	
+	/**
+	 * Converts path to FText.
+	 * @return FText representation of path.
+	 */
+	 FText ToText() const;
+};
+
+/**
+ * Allows RAII scoped control of FJsonSchemaMemberPath. 
+ * Pushes string on construction, pops string on destruction.
+ * Can be used directly, or with macro below.
+ */
+struct FJsonSchemaScopedMemberPathPush
+{
+	UE_NONCOPYABLE(FJsonSchemaScopedMemberPathPush)
+
+	explicit FJsonSchemaScopedMemberPathPush(FJsonSchemaMemberPath& InMemberPath, FStringView InString, const bool bInIsSubmember) : 
+		MemberPath(InMemberPath),
+		String(InString),
+		bIsSubmember(bInIsSubmember)
+	{
+		if (bIsSubmember)
+		{
+			MemberPath.PushSubmember(String);
+		}
+		else
+		{
+			MemberPath.Push(String);
+		}
+	}
+
+	~FJsonSchemaScopedMemberPathPush()
+	{
+		if (bIsSubmember)
+		{
+			MemberPath.PopSubmember(String);
+		}
+		else
+		{
+			MemberPath.Pop(String);
+		}
+	}
+
+private:
+	
+	FJsonSchemaMemberPath& MemberPath;
+	FString String; // ..can't use FStringView, out of scope when dtr called
+	bool bIsSubmember;
+};
+
+// Declares uniquely named FScopedSubmemberPath.
+#define UE_JSON_SCHEMA_SCOPED_MEMBER_PATH_PUSH(MemberPathExpr, StringExpr, bIsSubmemberExpr) \
+	FJsonSchemaScopedMemberPathPush UE_JOIN(_ScopedSubmemberPath_, __LINE__)((MemberPathExpr), (StringExpr), (bIsSubmemberExpr));
diff --git a/Engine/Source/Runtime/JsonUtilities/Public/JsonSchema/JsonSchemaPropertyFilter.h b/Engine/Source/Runtime/JsonUtilities/Public/JsonSchema/JsonSchemaPropertyFilter.h
new file mode 100644
--- /dev/null
+++ b/Engine/Source/Runtime/JsonUtilities/Public/JsonSchema/JsonSchemaPropertyFilter.h
@@ -0,0 +1,81 @@
+// Copyright Epic Games, Inc. All Rights Reserved. 
+
+#pragma once
+
+#include "CoreMinimal.h"
+
+#include "JsonObjectConverter.h"
+
+/**
+ * Filter used during schema generation, used for both FProperties and UStructs.
+ */
+struct FJsonSchemaPropertyFilter
+{
+	JSONUTILITIES_API explicit FJsonSchemaPropertyFilter(
+		const EPropertyFlags InCheckFlags = CPF_None,
+		const EPropertyFlags InSkipFlags = CPF_None,
+		const TOptional<TSet<FString>>& InRequiredPropertyMemberPaths = TOptional<TSet<FString>>(),
+		const TOptional<TSet<FString>>& InSkipPropertyMemberPaths = TOptional<TSet<FString>>(),
+		const EJsonObjectConversionFlags InConversionFlags = EJsonObjectConversionFlags::None);
+	
+	JSONUTILITIES_API FJsonSchemaPropertyFilter(const FJsonSchemaPropertyFilter& Other);
+	
+	/**
+	 * If specified (not CPF_None), only properties with any of the specified EPropertyFlags will be considered.
+	 */
+	EPropertyFlags CheckFlags = CPF_None;
+
+	/**
+	 * If specified (not CPF_None), any properties with any of the specified EPropertyFlags will be skipped.
+	 */
+	EPropertyFlags SkipFlags = CPF_None;
+
+	/**
+	 * Only considered while generating schemas for UStructs.
+	 * If specified, only properties in this list will be declared in the 'required' schema attribute.
+	 * Inner properties should be specified by path e.g: ["StructMember", "StructMember.InnerProperty"]
+	 * (If unspecified, all properties without default values in StructMetaData will end up automatically 
+	 * listed as required, in code found elsewhere.)
+	 */
+	TOptional<TSet<FString>> RequiredPropertyMemberPaths;
+
+	/**
+	 * Only considered while generating schemas for UStructs.
+	 * If specified, any properties in this list will not be described in the resulting schema.
+	 * Inner properties should be specified by path e.g: ["StructMember", "StructMember.InnerProperty"]
+	 */
+	TOptional<TSet<FString>> SkipPropertyMemberPaths;
+
+	/**
+	 * If passed, ensures generated schema is compatible with FJsonObjectConverter::UStructToJsonObject executed with the
+	 * same flags e.g: EJsonObjectConversionFlags::SkipStandardizeCase for property name casing.
+	 */
+	EJsonObjectConversionFlags ConversionFlags = EJsonObjectConversionFlags::None;
+	
+	/**
+	 * Checks if a property should be ignored, to its flags compared to our check/skip flags.
+	 * @param Property The property to test.
+	 * @return Whether a property should be ignored.
+	 */
+	JSONUTILITIES_API bool IsPropertyIgnored(TNotNull<const FProperty*> Property) const;
+
+	/**
+	 * Checks if a property is required.
+	 * @param PropertyMemberPath 
+	 * @return Whether the property path is listed as being required.
+	 */
+	JSONUTILITIES_API bool IsPropertyMemberPathRequired(const FString& PropertyMemberPath) const;
+
+	/**
+	 * Checks if a property should be skipped.
+	 * @param PropertyMemberPath 
+	 * @return Whether the property path is listed as being skipped.
+	 */
+	JSONUTILITIES_API bool IsPropertyMemberPathSkipped(const FString& PropertyMemberPath) const;
+
+	/**
+	 * Converts a property's authored name into a name that respects the ConversionFlags setting.
+	 * @param PropertyAuthoredName The authored name of a property.
+	 */
+	JSONUTILITIES_API FString PropertyAuthoredNameToJsonKey(const FString& PropertyAuthoredName) const;
+};
