JsonUtilities JsonSchema.

(1) Unbounded recursion fix -
Using a simple depth-first stack solution to avoid unbounded function recursion for schema generation, on both Runtime and Editor side.
Added simple visitor stack element class.
(Not using PropertyVisitor, as code ended up more complex, and actual property data is not needed for generating schema.)

(2) (Minor) Linkage fix -
FJsonSchemaMemberPath operator=() has correct linkage.

(3) (Minor) Module Feature -
Module Features return metadata, instead of taking the output metadata as a function parameter.

#Synced-CL 49088443

diff --git a/Engine/Source/Editor/JsonUtilitiesEditor/Private/JsonSchema/JsonSchemaGeneratorEditor.cpp b/Engine/Source/Editor/JsonUtilitiesEditor/Private/JsonSchema/JsonSchemaGeneratorEditor.cpp
--- a/Engine/Source/Editor/JsonUtilitiesEditor/Private/JsonSchema/JsonSchemaGeneratorEditor.cpp
+++ b/Engine/Source/Editor/JsonUtilitiesEditor/Private/JsonSchema/JsonSchemaGeneratorEditor.cpp
@@ -5,9 +5,11 @@
 #include "UObject/UnrealType.h"
 #include "EdGraphSchema_K2.h"
 #include "ObjectTools.h"
-
+#include "JsonSchema/JsonSchemaVisitorStack.h"
 #include "JsonSchema/JsonSchemaGenerator.h"
 
+using namespace UE::JsonSchema;
+
 //
 // Internal FJsonSchemaMetadataVisitor
 //
@@ -16,17 +18,45 @@ class FJsonSchemaMetadataVisitor
 {
 public:
 
-	explicit FJsonSchemaMetadataVisitor(
-			FJsonSchemaEditorMetadata& InEditorMetadata,
+	explicit FJsonSchemaMetadataVisitor(FJsonSchemaEditorMetadata& InOutputEditorMetadata,
 			const FJsonSchemaPropertyFilter& InPropertyFilter = FJsonSchemaPropertyFilter()) :
-		EditorMetadata(InEditorMetadata),
-		PropertyFilter(InPropertyFilter)
+		OutputEditorMetadata(InOutputEditorMetadata)
 	{
 	}
-	
-	void VisitFProperty(TNotNull<const FProperty*> Property,
-		const bool bIsPropertyFunctionWorldContext = false)
+		
+	void VisitRoot(const FVisitorTarget& VisitorTarget, 
+		const FJsonSchemaMemberPath& MemberPath, const FJsonSchemaPropertyFilter& PropertyFilter)
+	{
+		const EVisitorStackElementFlags Flags = (VisitorTarget.IsType<FValidConstUStruct>() ? 
+			EVisitorStackElementFlags::IsRootStruct : EVisitorStackElementFlags::None);
+    		
+		const FVisitorStackElement RootStackElement(VisitorTarget, Flags, MemberPath, PropertyFilter, nullptr); 
+		VisitorStack.Add(RootStackElement);
+		
+		while (!VisitorStack.IsEmpty())
+		{
+			const FVisitorStackElement StackElement = VisitorStack.Pop();
+			
+			if (StackElement.VisitorTarget.IsType<FValidConstUStruct>())
+			{
+				VisitUStruct(StackElement.VisitorTarget.Get<FValidConstUStruct>(), 
+					StackElement.MemberPath, StackElement.PropertyFilter, 
+					StackElement.Flags);
+			}
+			else if (StackElement.VisitorTarget.IsType<FValidConstFProperty>())
+			{
+				VisitFProperty(StackElement.VisitorTarget.Get<FValidConstFProperty>(),
+					StackElement.MemberPath, StackElement.PropertyFilter, StackElement.Flags);
+			}
+		}
+	}
+
+	void VisitFProperty(const FValidConstFProperty Property, 
+		const FJsonSchemaMemberPath& MemberPath, const FJsonSchemaPropertyFilter& PropertyFilter,
+		const EVisitorStackElementFlags Flags)
 	{
+		OutputEditorMetadata.CurrentPropertyMemberPath = MemberPath;
+		
 		const UFunction* PropertyOwnerFunction = Cast<UFunction>(Property->GetOwnerStruct());
 			
 		// Check to see if we should ignore this property.
@@ -77,30 +107,35 @@ public:
 			}
 		}
 
-		if (PropertyOwnerFunction && bIsPropertyFunctionWorldContext)
+		if (PropertyOwnerFunction && EnumHasAnyFlags(Flags, EVisitorStackElementFlags::IsFunctionWorldContextProperty))
 		{
 			PropertyMetadata->SetBoolField(TEXT("isWorldContext"), true);
 		}
 		
 		if (PropertyMetadata->Values.Num())
 		{
-			EditorMetadata.SetPropertyMetadataForCurrentPropertyMemberPath(PropertyMetadata);
+			OutputEditorMetadata.SetPropertyMetadataForCurrentPropertyMemberPath(PropertyMetadata);
 		}
 
-		// Recurse into inner properties
-		if (const UStruct* InnerStruct = GetFPropertyInnerStruct(*Property, *EditorMetadata.CurrentPropertyMemberPath))
+		// Recurse into inner properties.
+		if (const UStruct* InnerStruct = GetFPropertyInnerStruct(*Property, *OutputEditorMetadata.CurrentPropertyMemberPath))
 		{
 			// See all NOTE_JSON_SCHEMA_CHECK_FLAGS_NO_PARMS.
-			TGuardValue CheckFlagsGuard(PropertyFilter.CheckFlags, PropertyFilter.CheckFlags & ~CPF_ParmFlags);
+			FJsonSchemaPropertyFilter NewPropertyFilter = PropertyFilter;
+			NewPropertyFilter.CheckFlags &= ~CPF_ParmFlags;
 
-			UE_JSON_SCHEMA_SCOPED_MEMBER_PATH_PUSH(EditorMetadata.CurrentPropertyMemberPath, FJsonSchemaMemberPath::Delimiter, false);
-			
-			VisitUStruct(InnerStruct, false);
+			VisitorStack.Emplace(
+				FVisitorTarget(TInPlaceType<FValidConstUStruct>(), InnerStruct), 
+				EVisitorStackElementFlags::None, OutputEditorMetadata.CurrentPropertyMemberPath, NewPropertyFilter, nullptr);
 		}
 	}
 	
-	void VisitUStruct(const TNotNull<const UStruct*> Struct, const bool bIsRootStruct)
+	void VisitUStruct(const FValidConstUStruct Struct, 
+		const FJsonSchemaMemberPath& MemberPath, const FJsonSchemaPropertyFilter& PropertyFilter,
+		const EVisitorStackElementFlags Flags)
 	{
+		OutputEditorMetadata.CurrentPropertyMemberPath = MemberPath;
+		
 		const UFunction* StructAsFunction = Cast<UFunction>(Struct);
 		
 		const FString FunctionWorldContext =
@@ -108,9 +143,9 @@ public:
 			StructAsFunction->GetMetaData(FBlueprintMetadata::MD_WorldContext) :
 			"";
 		
-		if (bIsRootStruct)
+		if (EnumHasAnyFlags(Flags, EVisitorStackElementFlags::IsRootStruct))
 		{
-			EditorMetadata.RootStructMetadata = MakeShared<FJsonObject>();
+			OutputEditorMetadata.RootStructMetadata = MakeShared<FJsonObject>();
 			
 			if (const FString Description = (
 					StructAsFunction ?
@@ -118,12 +153,12 @@ public:
 					Struct->GetToolTipText().ToString());
 				!Description.IsEmpty())
 			{
-				EditorMetadata.RootStructMetadata->SetStringField(TEXT("description"), Description);
+				OutputEditorMetadata.RootStructMetadata->SetStringField(TEXT("description"), Description);
 			}
 			
 			if (!FunctionWorldContext.IsEmpty())
 			{
-				EditorMetadata.RootStructMetadata->SetStringField(TEXT("worldContext"), FunctionWorldContext);
+				OutputEditorMetadata.RootStructMetadata->SetStringField(TEXT("worldContext"), FunctionWorldContext);
 			}
 		}
 		
@@ -131,11 +166,16 @@ public:
 		{
 			const FString PropertyName = Property->GetAuthoredName();
 
-			UE_JSON_SCHEMA_SCOPED_MEMBER_PATH_PUSH(EditorMetadata.CurrentPropertyMemberPath, PropertyName, false);
+			UE_JSON_SCHEMA_SCOPED_MEMBER_PATH_PUSH(OutputEditorMetadata.CurrentPropertyMemberPath, PropertyName, 
+				!EnumHasAnyFlags(Flags, EVisitorStackElementFlags::IsRootStruct));
 			
-			const bool bIsPropertyFunctionWorldContext = (FunctionWorldContext == PropertyName);
+			const EVisitorStackElementFlags NewFlags = (FunctionWorldContext == PropertyName ?
+				EVisitorStackElementFlags::IsFunctionWorldContextProperty : 
+				EVisitorStackElementFlags::None);
 			
-			VisitFProperty(Property,	bIsPropertyFunctionWorldContext);
+			VisitorStack.Emplace(
+				FVisitorTarget(TInPlaceType<FValidConstFProperty>(), Property), 
+				NewFlags, OutputEditorMetadata.CurrentPropertyMemberPath, PropertyFilter, nullptr);
 		}
 	}
 	
@@ -177,30 +217,39 @@ private:
 		// @todo Push into FObjectProperty via struct instance?
 		return nullptr;
 	}
-
-	FJsonSchemaEditorMetadata& EditorMetadata; 
 	
-	FJsonSchemaPropertyFilter PropertyFilter;
+	FVisitorStack VisitorStack;
+	FJsonSchemaEditorMetadata& OutputEditorMetadata; 
 };
 
 //
 // FJsonSchemaGeneratorEditor
 //
 
-/*static*/ void FJsonSchemaGeneratorEditor::FPropertyToJsonSchemaMetadata(TNotNull<const FProperty*> Property,
-	FJsonSchemaEditorMetadata& OutEditorMetadata,
+/*static*/ FJsonSchemaEditorMetadata FJsonSchemaGeneratorEditor::FPropertyToJsonSchemaMetadata(TNotNull<const FProperty*> Property,
 	const FJsonSchemaPropertyFilter& PropertyFilter)
 {
-	UE_JSON_SCHEMA_SCOPED_MEMBER_PATH_PUSH(OutEditorMetadata.CurrentPropertyMemberPath, Property->GetAuthoredName(), false);
-
-	FJsonSchemaMetadataVisitor JsonSchemaMetadataVisitor(OutEditorMetadata, PropertyFilter);
-	JsonSchemaMetadataVisitor.VisitFProperty(Property);
+	FJsonSchemaEditorMetadata EditorMetadata;
+	
+	FJsonSchemaMetadataVisitor JsonSchemaMetadataVisitor(EditorMetadata, PropertyFilter);
+	JsonSchemaMetadataVisitor.VisitRoot(
+		FVisitorTarget(TInPlaceType<FValidConstFProperty>(), Property), 
+		FJsonSchemaMemberPath(*Property->GetAuthoredName()),
+		PropertyFilter);
+	
+	return EditorMetadata;
 }
 
-/*static*/ void FJsonSchemaGeneratorEditor::UStructToJsonSchemaMetadata(TNotNull<const UStruct*> Struct,
-	FJsonSchemaEditorMetadata& OutEditorMetadata,
+/*static*/ FJsonSchemaEditorMetadata FJsonSchemaGeneratorEditor::UStructToJsonSchemaMetadata(TNotNull<const UStruct*> Struct,
 	const FJsonSchemaPropertyFilter& PropertyFilter)
 {
-	FJsonSchemaMetadataVisitor JsonSchemaMetadataVisitor(OutEditorMetadata, PropertyFilter);
-	JsonSchemaMetadataVisitor.VisitUStruct(Struct, true);
+	FJsonSchemaEditorMetadata EditorMetadata;
+	
+	FJsonSchemaMetadataVisitor JsonSchemaMetadataVisitor(EditorMetadata, PropertyFilter);
+	JsonSchemaMetadataVisitor.VisitRoot(
+		FVisitorTarget(TInPlaceType<FValidConstUStruct>(), Struct), 
+		FJsonSchemaMemberPath(),
+		PropertyFilter);
+	
+	return EditorMetadata;
 }
diff --git a/Engine/Source/Editor/JsonUtilitiesEditor/Private/JsonUtilitiesEditorModularFeature.cpp b/Engine/Source/Editor/JsonUtilitiesEditor/Private/JsonUtilitiesEditorModularFeature.cpp
--- a/Engine/Source/Editor/JsonUtilitiesEditor/Private/JsonUtilitiesEditorModularFeature.cpp
+++ b/Engine/Source/Editor/JsonUtilitiesEditor/Private/JsonUtilitiesEditorModularFeature.cpp
@@ -6,22 +6,16 @@
 
 PRAGMA_DISABLE_DEPRECATION_WARNINGS
 
-/*virtual*/ void FJsonUtilitiesEditorModularFeature::FPropertyToJsonSchemaMetadata(TNotNull<const FProperty*> Property,
-	FJsonSchemaEditorMetadata& OutEditorMetadata,
+/*virtual*/ FJsonSchemaEditorMetadata FJsonUtilitiesEditorModularFeature::FPropertyToJsonSchemaMetadata(TNotNull<const FProperty*> Property,
 	const FJsonSchemaPropertyFilter& PropertyFilter) const
 {
-	FJsonSchemaGeneratorEditor::FPropertyToJsonSchemaMetadata(Property,
-		OutEditorMetadata,
-		PropertyFilter);
+	return FJsonSchemaGeneratorEditor::FPropertyToJsonSchemaMetadata(Property, PropertyFilter);
 }
 
-/*virtual*/ void FJsonUtilitiesEditorModularFeature::UStructToJsonSchemaMetadata(TNotNull<const UStruct*> Struct,
-	FJsonSchemaEditorMetadata& OutEditorMetadata,
+/*virtual*/ FJsonSchemaEditorMetadata FJsonUtilitiesEditorModularFeature::UStructToJsonSchemaMetadata(TNotNull<const UStruct*> Struct,
 	const FJsonSchemaPropertyFilter& PropertyFilter) const
 {
-	return FJsonSchemaGeneratorEditor::UStructToJsonSchemaMetadata(Struct,
-		OutEditorMetadata,
-		PropertyFilter);
+	return FJsonSchemaGeneratorEditor::UStructToJsonSchemaMetadata(Struct, PropertyFilter);
 }
 
 PRAGMA_ENABLE_DEPRECATION_WARNINGS
diff --git a/Engine/Source/Editor/JsonUtilitiesEditor/Private/JsonUtilitiesEditorModularFeature.h b/Engine/Source/Editor/JsonUtilitiesEditor/Private/JsonUtilitiesEditorModularFeature.h
--- a/Engine/Source/Editor/JsonUtilitiesEditor/Private/JsonUtilitiesEditorModularFeature.h
+++ b/Engine/Source/Editor/JsonUtilitiesEditor/Private/JsonUtilitiesEditorModularFeature.h
@@ -12,11 +12,9 @@ class FJsonUtilitiesEditorModularFeature : public IJsonUtilitiesModularFeature
 {
 public:
 
-	JSONUTILITIESEDITOR_API virtual void FPropertyToJsonSchemaMetadata(TNotNull<const FProperty*> Property,
-		FJsonSchemaEditorMetadata& OutEditorMetadata,
+	JSONUTILITIESEDITOR_API virtual FJsonSchemaEditorMetadata FPropertyToJsonSchemaMetadata(TNotNull<const FProperty*> Property,
 		const FJsonSchemaPropertyFilter& PropertyFilter = FJsonSchemaPropertyFilter()) const override;
 
-	JSONUTILITIESEDITOR_API virtual void UStructToJsonSchemaMetadata(TNotNull<const UStruct*>,
-		FJsonSchemaEditorMetadata& OutEditorMetadata,
+	JSONUTILITIESEDITOR_API virtual	FJsonSchemaEditorMetadata UStructToJsonSchemaMetadata(TNotNull<const UStruct*>,
 		const FJsonSchemaPropertyFilter& PropertyFilter = FJsonSchemaPropertyFilter()) const override;
 };
diff --git a/Engine/Source/Editor/JsonUtilitiesEditor/Public/JsonSchema/JsonSchemaGeneratorEditor.h b/Engine/Source/Editor/JsonUtilitiesEditor/Public/JsonSchema/JsonSchemaGeneratorEditor.h
--- a/Engine/Source/Editor/JsonUtilitiesEditor/Public/JsonSchema/JsonSchemaGeneratorEditor.h
+++ b/Engine/Source/Editor/JsonUtilitiesEditor/Public/JsonSchema/JsonSchemaGeneratorEditor.h
@@ -17,18 +17,17 @@ public:
 	 * Usually not called by a developer directly. Used by the Runtime module via the Module Feature.
 	 * 
 	 * @param Property The property to collect metadata for.
-	 * @param OutEditorMetadata The collected property metadata.
 	 * @param PropertyFilter Determines how to filter collected struct and property data.
+	 * @return The collected property metadata.
 	 */
-	JSONUTILITIESEDITOR_API static void FPropertyToJsonSchemaMetadata(TNotNull<const FProperty*> Property,
-		FJsonSchemaEditorMetadata& OutEditorMetadata,
+	JSONUTILITIESEDITOR_API static FJsonSchemaEditorMetadata FPropertyToJsonSchemaMetadata(TNotNull<const FProperty*> Property,
 		const FJsonSchemaPropertyFilter& PropertyFilter = FJsonSchemaPropertyFilter());
 
 	/**
 	 * Usually not called by a developer directly. Used by the Runtime module via the Module Feature.
 	 * 
 	 * Recursively walks Struct's property hierarchy, collecting editor-only meta-data like
-	 * tooltips and other Meta values, caching them in the cookable OutEditorMetadata.
+	 * tooltips and other Meta values, returning them in the cookable metadata.
 	 *
 	 * e.g:
 	 *
@@ -75,10 +74,9 @@ public:
 	 *			PropertyMetadata -> Person -> "The person to find the best friend of"
 	 *
 	 * @param Struct The type to collect property data for. Can be either UScriptStruct, UClass or UFunction.
-	 * @param OutEditorMetadata The collected property metadata.
 	 * @param PropertyFilter Determines how to filter collected struct and property data.
+	 * @return The collected property metadata.
 	 */
-	JSONUTILITIESEDITOR_API static void UStructToJsonSchemaMetadata(TNotNull<const UStruct*> Struct,
-		FJsonSchemaEditorMetadata& OutEditorMetadata,
+	JSONUTILITIESEDITOR_API static FJsonSchemaEditorMetadata UStructToJsonSchemaMetadata(TNotNull<const UStruct*> Struct,
 		const FJsonSchemaPropertyFilter& PropertyFilter = FJsonSchemaPropertyFilter());
 };
diff --git a/Engine/Source/Runtime/JsonUtilities/Private/JsonSchema/JsonSchemaGenerator.cpp b/Engine/Source/Runtime/JsonUtilities/Private/JsonSchema/JsonSchemaGenerator.cpp
--- a/Engine/Source/Runtime/JsonUtilities/Private/JsonSchema/JsonSchemaGenerator.cpp
+++ b/Engine/Source/Runtime/JsonUtilities/Private/JsonSchema/JsonSchemaGenerator.cpp
@@ -7,15 +7,37 @@
 #include "UObject/TextProperty.h"
 #include "UObject/UnrealType.h"
 #include "Policies/CondensedJsonPrintPolicy.h"
+#include "JsonSchema/JsonSchemaVisitorStack.h"
 
 #include "JsonUtilitiesModularFeature.h"
 
+using namespace UE::JsonSchema;
+
 // See all NOTE_JSON_SCHEMA_CHECK_FLAGS_NO_PARMS.
 // A UStruct may be a UFunction, in which case a caller might have set the filter flags to ignore
 // certain types of parameters (properties) of the function. After using the CheckFlags to check
 // of a property/parameter should be ignored, we should stop considering flags related to function
 // parameters (properties), by clearing all CPF_ParmFlags, so that they don't block schema
-// generation (or metadata collection) for sub-structs/sub-properties.
+// generation (or metadata collection) for substructs/subproperties.
+
+//
+// Internal
+//
+
+namespace 
+{
+	template<class CharType, class PrintPolicy>
+	FString JsonObjectToJsonObjectString(const TSharedRef<FJsonObject>& JsonObject, const int32 Indent)
+	{
+		FString OutString;
+		
+		TSharedRef<TJsonWriter<CharType, PrintPolicy>> JsonWriter = TJsonWriterFactory<CharType, PrintPolicy>::Create(&OutString, Indent);
+		const bool bSuccess = FJsonSerializer::Serialize(JsonObject, JsonWriter);
+		JsonWriter->Close();
+		
+		return (bSuccess ? OutString : TEXT(""));
+	}
+}
 
 //
 // Internal FJsonSchemaGeneratorVisitor
@@ -25,105 +47,135 @@ class FJsonSchemaGeneratorVisitor
 {
 public:
 
-	explicit FJsonSchemaGeneratorVisitor(
-			FJsonSchemaEditorMetadata& InEditorMetadata,
-			const FJsonSchemaPropertyFilter& InPropertyFilter = FJsonSchemaPropertyFilter()) :
-		EditorMetadata(InEditorMetadata),
-		PropertyFilter(InPropertyFilter)
+	explicit FJsonSchemaGeneratorVisitor(FJsonSchemaEditorMetadata& InEditorMetadata) :
+		EditorMetadata(InEditorMetadata)
 	{
 	}
 	
-	TSharedPtr<FJsonObject> VisitFProperty(TNotNull<const FProperty*> Property)
+	TSharedPtr<FJsonObject> VisitRoot(const FVisitorTarget& VisitorTarget, 
+		const FJsonSchemaMemberPath& MemberPath, const FJsonSchemaPropertyFilter& PropertyFilter)
+	{
+		const EVisitorStackElementFlags Flags = (VisitorTarget.IsType<FValidConstUStruct>() ? 
+			EVisitorStackElementFlags::IsRootStruct : EVisitorStackElementFlags::None);
+		
+		const FVisitorStackElement RootStackElement(VisitorTarget, Flags, MemberPath, PropertyFilter, MakeShared<FJsonObject>()); 
+		VisitorStack.Add(RootStackElement);
+		
+		while (!VisitorStack.IsEmpty())
+		{
+			const FVisitorStackElement StackElement = VisitorStack.Pop();
+			
+			if (StackElement.VisitorTarget.IsType<FValidConstUStruct>())
+			{
+				VisitUStruct(StackElement.VisitorTarget.Get<FValidConstUStruct>(), 
+					StackElement.OutputSchema.ToSharedRef(), StackElement.MemberPath, StackElement.PropertyFilter, 
+					StackElement.Flags);
+			}
+			else if (StackElement.VisitorTarget.IsType<FValidConstFProperty>())
+			{
+				VisitFProperty(StackElement.VisitorTarget.Get<FValidConstFProperty>(), 
+					StackElement.OutputSchema.ToSharedRef(), StackElement.MemberPath, StackElement.PropertyFilter);
+			}
+		}
+	
+		return RootStackElement.OutputSchema;
+	}
+	
+	void VisitFProperty(const FValidConstFProperty Property, const TSharedRef<FJsonObject>& OutputSchema, 
+		const FJsonSchemaMemberPath& MemberPath, const FJsonSchemaPropertyFilter& PropertyFilter)
 	{
-		FJsonDomBuilder::FObject PropertyJson;
+		EditorMetadata.CurrentPropertyMemberPath = MemberPath;
 		
 		if (const FOptionalProperty* OptionalProperty = CastField<FOptionalProperty>(Property))
 		{
-			BuildJsonForOptionalProperty(*OptionalProperty, PropertyJson);
+			BuildJsonForOptionalProperty(*OptionalProperty, PropertyFilter, OutputSchema);
 		}
 		else if (const FStructProperty* StructProperty = CastField<FStructProperty>(Property))
 		{
-			BuildJsonForStructProperty(*StructProperty, PropertyJson);
+			BuildJsonForStructProperty(*StructProperty, PropertyFilter, OutputSchema);
 		}
 		else if (const FArrayProperty* ArrayProperty = CastField<FArrayProperty>(Property))
 		{
-			BuildJsonForArrayProperty(*ArrayProperty, PropertyJson);
+			BuildJsonForArrayProperty(*ArrayProperty, PropertyFilter, OutputSchema);
 		}
 		else if (const FSetProperty* SetProperty = CastField<FSetProperty>(Property))
 		{
-			BuildJsonForSetProperty(*SetProperty, PropertyJson);
+			BuildJsonForSetProperty(*SetProperty, PropertyFilter, OutputSchema);
 		}
 		else if (const FMapProperty* MapProperty = CastField<FMapProperty>(Property))
 		{
-			BuildJsonForMapProperty(*MapProperty, PropertyJson);
+			BuildJsonForMapProperty(*MapProperty, PropertyFilter, OutputSchema);
 		}
 		else if (const FEnumProperty* EnumProperty = CastField<FEnumProperty>(Property))
 		{
 			if (EnumProperty->GetEnum())
 			{
-				BuildJsonForEnumDefinition(*EnumProperty->GetEnum(), PropertyJson);
+				BuildJsonForEnumDefinition(*EnumProperty->GetEnum(), OutputSchema);
 			}
 		}
 		else if (const FNumericProperty* NumericProperty = CastField<FNumericProperty>(Property))
 		{
-			BuildJsonForNumericProperty(*NumericProperty, PropertyJson);
+			BuildJsonForNumericProperty(*NumericProperty, OutputSchema);
 		}
 		else if (CastField<FBoolProperty>(Property))
 		{
-			BuildJsonForBoolProperty(PropertyJson);
+			BuildJsonForBoolProperty(OutputSchema);
 		}
 		else if (CastField<FStrProperty>(Property) || CastField<FTextProperty>(Property) || CastField<FNameProperty>(Property))
 		{
-			BuildJsonForStringProperty(PropertyJson);
+			BuildJsonForStringProperty(OutputSchema);
 		}
 		else
 		{
 			// Fall back to string for unhandled property types.
-			BuildJsonForStringProperty(PropertyJson);
+			BuildJsonForStringProperty(OutputSchema);
 		}
 
 		// Add all property metadata, if any.
-		if (const TSharedPtr<FJsonObject> PropertyMetadata = EditorMetadata.GetPropertyMetadataForCurrentPropertyMemberPath())
+		if (const TSharedPtr<FJsonObject> PropertyMetadataJson = EditorMetadata.GetPropertyMetadataForCurrentPropertyMemberPath())
 		{
-			for (const TTuple<FString, TSharedPtr<FJsonValue>>& Tuple : PropertyMetadata->Values)
+			for (const TTuple<FString, TSharedPtr<FJsonValue>>& Tuple : PropertyMetadataJson->Values)
 			{
 				const FString JsonMemberName = PropertyFilter.PropertyAuthoredNameToJsonKey(Tuple.Key);
-				PropertyJson.Set(JsonMemberName, Tuple.Value);
+				OutputSchema->SetField(*JsonMemberName, Tuple.Value);
 			}
 		}
-		
-		return PropertyJson.AsJsonObject();
 	}
 	
-	TSharedPtr<FJsonObject> VisitUStruct(TNotNull<const UStruct*> Struct, const bool bIsRootStruct)
+	void VisitUStruct(const FValidConstUStruct Struct, const TSharedRef<FJsonObject>& OutputSchema, 
+		const FJsonSchemaMemberPath& MemberPath, const FJsonSchemaPropertyFilter& PropertyFilter,
+		const EVisitorStackElementFlags Flags)
 	{
-		FJsonDomBuilder::FObject OutStructJson;
-
-		OutStructJson.Set(TEXT("type"), TEXT("object"));
+		EditorMetadata.CurrentPropertyMemberPath = MemberPath;
+		
+		OutputSchema->SetStringField(TEXT("type"), TEXT("object"));
 
 		// Special case handling for FJsonObjectWrapper, where we know this must be a Json object.   
 		if (Struct == FJsonObjectWrapper::StaticStruct())
 		{
-			return OutStructJson.AsJsonObject();
+			return;
 		}
 
 		// Add all struct metadata.
-		if (bIsRootStruct && EditorMetadata.RootStructMetadata)
+		if (EnumHasAnyFlags(Flags, EVisitorStackElementFlags::IsRootStruct) && EditorMetadata.RootStructMetadata)
 		{
 			for (const TTuple<FString, TSharedPtr<FJsonValue>>& Tuple : EditorMetadata.RootStructMetadata->Values)
 			{
 				const FString JsonMemberName = PropertyFilter.PropertyAuthoredNameToJsonKey(Tuple.Key);
-				OutStructJson.Set(JsonMemberName, Tuple.Value);
+				OutputSchema->SetField(JsonMemberName, Tuple.Value);
 			}
 		}
 
-		// Adds property to JSON and possibly the require list. Returns whether the property was added to the JSON object.
-		auto MaybeAddPropertyToJson = [this](const FProperty* Property,
-			FJsonDomBuilder::FObject& PropertiesJson, FJsonDomBuilder::FArray& RequiredPropertiesArray) -> bool
+		// Adds property to JSON and possibly the required list.
+		// If added, property will also be added to the visitor stack for processing.
+		// Returns whether the property was added to the JSON object.
+		auto MaybeAddProperty = [this, PropertyFilter, Flags](const FProperty* Property,
+			const TSharedRef<FJsonObject>& PropertiesSchema, TArray<TSharedPtr<FJsonValue>>& RequiredPropertiesArray) -> bool
 			{
 				const FString PropertyName = Property->GetAuthoredName();
 			
-				UE_JSON_SCHEMA_SCOPED_MEMBER_PATH_PUSH(EditorMetadata.CurrentPropertyMemberPath, PropertyName, false);
+				UE_JSON_SCHEMA_SCOPED_MEMBER_PATH_PUSH(EditorMetadata.CurrentPropertyMemberPath, PropertyName, 
+					!EnumHasAnyFlags(Flags, EVisitorStackElementFlags::IsRootStruct));
 
 				if (PropertyFilter.IsPropertyIgnored(Property) || 
 					PropertyFilter.IsPropertyMemberPathSkipped(*EditorMetadata.CurrentPropertyMemberPath))
@@ -132,38 +184,42 @@ public:
 				}
 
 				// See all NOTE_JSON_SCHEMA_CHECK_FLAGS_NO_PARMS.
-				TGuardValue CheckFlagsGuard(PropertyFilter.CheckFlags, PropertyFilter.CheckFlags & ~CPF_ParmFlags);
+				FJsonSchemaPropertyFilter NewPropertyFilter = PropertyFilter;
+				NewPropertyFilter.CheckFlags &= ~CPF_ParmFlags;
 			
-				if (const TSharedPtr<FJsonObject> PropertySchema = VisitFProperty(Property);
-					PropertySchema && PropertySchema->Values.Num())
-				{
-					const FString JsonMemberName = PropertyFilter.PropertyAuthoredNameToJsonKey(PropertyName); 
-
-					PropertiesJson.Set(JsonMemberName, PropertySchema);
+				const TSharedRef<FJsonObject> PropertySchema = MakeShared<FJsonObject>();
+				
+				const FString JsonMemberName = PropertyFilter.PropertyAuthoredNameToJsonKey(PropertyName); 
+				PropertiesSchema->SetObjectField(JsonMemberName, PropertySchema);
+				
+				VisitorStack.Emplace(
+					FVisitorTarget(TInPlaceType<FValidConstFProperty>(), Property), 
+					EVisitorStackElementFlags::None, 
+					EditorMetadata.CurrentPropertyMemberPath, 
+					NewPropertyFilter,
+					PropertySchema);
 
-					const TSharedPtr<FJsonObject> PropertyMetadata = EditorMetadata.GetPropertyMetadataForCurrentPropertyMemberPath();
+				const TSharedPtr<FJsonObject> PropertyMetadata = EditorMetadata.GetPropertyMetadataForCurrentPropertyMemberPath();
 
-					if (PropertyFilter.IsPropertyMemberPathRequired(*EditorMetadata.CurrentPropertyMemberPath))
-					{
-						// Property has been explicitly required.
-						RequiredPropertiesArray.Add(JsonMemberName);
-					}
-					else if (!(PropertyMetadata && PropertyMetadata->HasField(TEXT("defaultValue"))) &&
-						!Property->IsA<FOptionalProperty>())
-					{
-						// Property has not been explicitly required.
-						// Only require non-optional properties that don't have default values.
-						RequiredPropertiesArray.Add(JsonMemberName);
-					}
-					
-					return true;
+				if (PropertyFilter.IsPropertyMemberPathRequired(*EditorMetadata.CurrentPropertyMemberPath))
+				{
+					// Property has been explicitly required.
+					RequiredPropertiesArray.Add(MakeShared<FJsonValueString>(JsonMemberName));
+				}
+				else if (!(PropertyMetadata && PropertyMetadata->HasField(TEXT("defaultValue"))) &&
+					!Property->IsA<FOptionalProperty>())
+				{
+					// Property has not been explicitly required.
+					// Only require non-optional properties that don't have default values.
+					RequiredPropertiesArray.Add(MakeShared<FJsonValueString>(JsonMemberName));
 				}
 				
-				return false;
+				return true;
 			};
 		
 		// Properties & required.
-		// If this struct is a function, use JSON RPC format (with inputSchema/outputSchema). Otherwise, use standard format.
+		// If this struct is a function, use JSON RPC format (with inputSchema/outputSchema). 
+		// Otherwise, use standard format.
 		if (Struct->IsA(UFunction::StaticClass()))
 		{
 			// If struct is a function, output JSON RPC -
@@ -172,87 +228,85 @@ public:
 			// [REDACTED_PROFILE_URL]/modelcontextprotocol/blob/main/schema/2025-06-18/schema.ts#L312
 			
 			// NOTE - We know we are a UFunction, so will have a valid Outer.
-			OutStructJson.Set(TEXT("name"), *(Struct->GetOuter()->GetName() + "." + Struct->GetAuthoredName()));
+			OutputSchema->SetStringField(TEXT("name"), *(Struct->GetOuter()->GetName() + "." + Struct->GetAuthoredName()));
 			
 			// Collect input and output properties.
-			FJsonDomBuilder::FObject InputSchemaPropertiesJson, OutputSchemaPropertiesJson;
-			FJsonDomBuilder::FArray InputSchemaRequiredPropertiesArray, OutputSchemaRequiredPropertiesArray;
+			const TSharedRef<FJsonObject> FunctionInputPropertiesSchema = MakeShared<FJsonObject>();
+			TArray<TSharedPtr<FJsonValue>> FunctionInputRequiredPropertiesArray;
+			const TSharedRef<FJsonObject> FunctionOutputPropertiesSchema = MakeShared<FJsonObject>();
+			TArray<TSharedPtr<FJsonValue>> FunctionOutputRequiredPropertiesArray;
 			for (const FProperty* Property : TFieldRange<FProperty>(Struct))
 			{
 				if (Property->HasAnyPropertyFlags(CPF_ReturnParm))
 				{
-					MaybeAddPropertyToJson(Property, OutputSchemaPropertiesJson, OutputSchemaRequiredPropertiesArray);
+					MaybeAddProperty(Property, FunctionOutputPropertiesSchema, FunctionOutputRequiredPropertiesArray);
 				}
 				else
 				{
-					MaybeAddPropertyToJson(Property, InputSchemaPropertiesJson, InputSchemaRequiredPropertiesArray);
+					MaybeAddProperty(Property, FunctionInputPropertiesSchema, FunctionInputRequiredPropertiesArray);
 				}
 			}
 	
 			// Input schema is not optional.
 			{
-				FJsonDomBuilder::FObject InputSchemaJson;
-				InputSchemaJson.Set(TEXT("type"), TEXT("object"));
-				OutStructJson.Set(TEXT("inputSchema"), InputSchemaJson);
+				const TSharedPtr<FJsonObject> FunctionInputSchema = MakeShared<FJsonObject>();
+				FunctionInputSchema->SetStringField(TEXT("type"), TEXT("object"));
+				OutputSchema->SetObjectField(TEXT("inputSchema"), FunctionInputSchema);
 				
-				if (InputSchemaPropertiesJson.Num() > 0)
+				if (FunctionInputPropertiesSchema->Values.Num() > 0)
 				{
-					InputSchemaJson.Set(TEXT("properties"), InputSchemaPropertiesJson);
+					FunctionInputSchema->SetObjectField(TEXT("properties"), FunctionInputPropertiesSchema);
 
-					if (InputSchemaRequiredPropertiesArray.Num() > 0)
+					if (FunctionInputRequiredPropertiesArray.Num() > 0)
 					{
-						InputSchemaJson.Set(TEXT("required"), InputSchemaRequiredPropertiesArray);
+						FunctionInputSchema->SetArrayField(TEXT("required"), FunctionInputRequiredPropertiesArray);
 					}
 				}
 			}
 	
 			// Output schema is optional.
-			if (OutputSchemaPropertiesJson.Num() > 0)
+			if (FunctionOutputPropertiesSchema->Values.Num() > 0)
 			{
-				FJsonDomBuilder::FObject OutputSchemaJson;
-				OutputSchemaJson.Set(TEXT("type"), TEXT("object"));
-				OutStructJson.Set(TEXT("outputSchema"), OutputSchemaJson);
+				const TSharedPtr<FJsonObject> FunctionOutputSchema = MakeShared<FJsonObject>();
+				FunctionOutputSchema->SetStringField(TEXT("type"), TEXT("object"));
+				OutputSchema->SetObjectField(TEXT("outputSchema"), FunctionOutputSchema);
 				
 				// We have OutputSchemaPropertiesJson (see above)
 				{
-					OutputSchemaJson.Set(TEXT("properties"), OutputSchemaPropertiesJson);
+					FunctionOutputSchema->SetObjectField(TEXT("properties"), FunctionOutputPropertiesSchema);
 					
-					if (OutputSchemaRequiredPropertiesArray.Num() > 0)
+					if (FunctionOutputRequiredPropertiesArray.Num() > 0)
 					{
-						OutputSchemaJson.Set(TEXT("required"), OutputSchemaRequiredPropertiesArray);
+						FunctionOutputSchema->SetArrayField(TEXT("required"), FunctionOutputRequiredPropertiesArray);
 					}
 				}
 			}
 		}
 		else
 		{
-			// If struct is not a function.
+			// If struct is not a function, output standard format.
 				
-			bool bDidAddProperties = false;
-			FJsonDomBuilder::FObject PropertiesJson;
-			FJsonDomBuilder::FArray RequiredPropertiesArray;
+			const TSharedRef<FJsonObject> PropertiesSchema = MakeShared<FJsonObject>();
+			TArray<TSharedPtr<FJsonValue>> RequiredPropertiesArray;
 			for (const FProperty* Property : TFieldRange<FProperty>(Struct))
 			{
-				bDidAddProperties |= MaybeAddPropertyToJson(Property, PropertiesJson, RequiredPropertiesArray);
+				MaybeAddProperty(Property, PropertiesSchema, RequiredPropertiesArray);
 			}
 	
-			if (bDidAddProperties)
+			if (PropertiesSchema->Values.Num() > 0)
 			{
-				OutStructJson.Set(TEXT("properties"), PropertiesJson);
+				OutputSchema->SetObjectField(TEXT("properties"), PropertiesSchema);
 				if (RequiredPropertiesArray.Num() > 0)
 				{
-					OutStructJson.Set(TEXT("required"), RequiredPropertiesArray);
+					OutputSchema->SetArrayField(TEXT("required"), RequiredPropertiesArray);
 				}
 			}
 		}
-
-		return OutStructJson.AsJsonObject();
 	}
 	
 private:
 	
-	static void BuildJsonForEnumDefinition(const UEnum& EnumDef,
-		FJsonDomBuilder::FObject& OutJson)
+	static void BuildJsonForEnumDefinition(const UEnum& EnumDef, const TSharedRef<FJsonObject>& OutputSchema)
 	{
 		FJsonDomBuilder::FArray EnumNames;
 		const int32 NumEnum = EnumDef.NumEnums();
@@ -266,117 +320,122 @@ private:
 			EnumNames.Add(StringValue);
 		}
 
-		OutJson.Set(TEXT("type"), TEXT("string"));
-		OutJson.Set(TEXT("enum"), EnumNames.AsJsonValue());
+		OutputSchema->SetStringField(TEXT("type"), TEXT("string"));
+		OutputSchema->SetField(TEXT("enum"), EnumNames.AsJsonValue());
 	}
 
-	void BuildJsonForOptionalProperty(const FOptionalProperty& OptionalProperty,
-		FJsonDomBuilder::FObject& OutJson)
+	void BuildJsonForOptionalProperty(const FOptionalProperty& OptionalProperty, 
+		const FJsonSchemaPropertyFilter& PropertyFilter, const TSharedRef<FJsonObject>& OutputSchema)
 	{
-		OutJson = VisitFProperty(OptionalProperty.GetValueProperty());
+		VisitorStack.Emplace(
+			FVisitorTarget(TInPlaceType<FValidConstFProperty>(), OptionalProperty.GetValueProperty()), 
+			EVisitorStackElementFlags::None, 
+			EditorMetadata.CurrentPropertyMemberPath, 
+			PropertyFilter,
+			OutputSchema);
 	}
 
-	void BuildJsonForStructProperty(const FStructProperty& StructProperty,
-		FJsonDomBuilder::FObject& OutJson)
+	void BuildJsonForStructProperty(const FStructProperty& StructProperty, 
+		const FJsonSchemaPropertyFilter& PropertyFilter, const TSharedRef<FJsonObject>& OutputSchema)
 	{
-		UE_JSON_SCHEMA_SCOPED_MEMBER_PATH_PUSH(EditorMetadata.CurrentPropertyMemberPath, FJsonSchemaMemberPath::Delimiter, false);
-
-		OutJson = VisitUStruct(StructProperty.Struct, false);
+		VisitorStack.Emplace(
+			FVisitorTarget(TInPlaceType<FValidConstUStruct>(), StructProperty.Struct),
+			EVisitorStackElementFlags::None,
+			EditorMetadata.CurrentPropertyMemberPath, 
+			PropertyFilter,
+			OutputSchema);
 	}
 
-	void BuildJsonForArrayProperty(const FArrayProperty& ArrayProperty,
-		FJsonDomBuilder::FObject& OutJson)
+	void BuildJsonForArrayProperty(const FArrayProperty& ArrayProperty, 
+		const FJsonSchemaPropertyFilter& PropertyFilter, const TSharedRef<FJsonObject>& OutputSchema)
 	{
-		const TSharedPtr<FJsonObject> ItemsJson = VisitFProperty(ArrayProperty.Inner);
+		const TSharedRef<FJsonObject> ItemsJson = MakeShared<FJsonObject>();
+		VisitorStack.Emplace(
+			FVisitorTarget(TInPlaceType<FValidConstFProperty>(), ArrayProperty.Inner), 
+			EVisitorStackElementFlags::None,
+			EditorMetadata.CurrentPropertyMemberPath, 
+			PropertyFilter,
+			ItemsJson);
 		
 		// Don't double-up the description for the inner items. Especially as the description will be written for the container e.g: "A list of ..."
 		ItemsJson->RemoveField(TEXT("description"));
 		
-		OutJson.Set(TEXT("type"), TEXT("array"));
-		OutJson.Set(TEXT("items"), ItemsJson);
+		OutputSchema->SetStringField(TEXT("type"), TEXT("array"));
+		OutputSchema->SetObjectField(TEXT("items"), ItemsJson);
 	}
 
-	void BuildJsonForSetProperty(const FSetProperty& SetProperty,
-		FJsonDomBuilder::FObject& OutJson)
+	void BuildJsonForSetProperty(const FSetProperty& SetProperty, 
+		const FJsonSchemaPropertyFilter& PropertyFilter, const TSharedRef<FJsonObject>& OutputSchema)
 	{
-		const TSharedPtr<FJsonObject> ItemsJson = VisitFProperty(SetProperty.ElementProp);
-		
+		const TSharedRef<FJsonObject> ItemsJson = MakeShared<FJsonObject>();
+		VisitorStack.Emplace(
+			FVisitorTarget(TInPlaceType<FValidConstFProperty>(), SetProperty.ElementProp), 
+			EVisitorStackElementFlags::None, 
+			EditorMetadata.CurrentPropertyMemberPath, 
+			PropertyFilter,
+			ItemsJson);
+
 		// Don't double-up the description for the inner items. Especially as the description will be written for the container e.g: "A list of ..."
 		ItemsJson->RemoveField(TEXT("description"));
 
-		OutJson.Set(TEXT("type"), TEXT("array"));
-		OutJson.Set(TEXT("uniqueItems"), true);
-		OutJson.Set(TEXT("items"), ItemsJson);
+		OutputSchema->SetStringField(TEXT("type"), TEXT("array"));
+		OutputSchema->SetBoolField(TEXT("uniqueItems"), true);
+		OutputSchema->SetObjectField(TEXT("items"), ItemsJson);
 	}
 
-	void BuildJsonForMapProperty(const FMapProperty& MapProperty,
-		FJsonDomBuilder::FObject& OutJson)
+	void BuildJsonForMapProperty(const FMapProperty& MapProperty, 
+		const FJsonSchemaPropertyFilter& PropertyFilter, const TSharedRef<FJsonObject>& OutputSchema)
 	{
 		if (ensureMsgf(MapProperty.KeyProp->IsA<FStrProperty>(), 
 			TEXT("BuildJsonForMapProperty() - Unsupported map key type for %s. JSON-compatible map keys must be FString. Map will be skipped."), 
 			*EditorMetadata.CurrentPropertyMemberPath))
 		{
-			const TSharedPtr<FJsonObject> MapValuesJson = VisitFProperty(MapProperty.ValueProp);
+			const TSharedRef<FJsonObject> MapValuesJson = MakeShared<FJsonObject>();
+			VisitorStack.Emplace(
+				FVisitorTarget(TInPlaceType<FValidConstFProperty>(), MapProperty.ValueProp), 
+				EVisitorStackElementFlags::None, 
+				EditorMetadata.CurrentPropertyMemberPath, 
+				PropertyFilter,
+				MapValuesJson);
 
 			// Don't double-up the description for the map values. Especially as the description will be written for the container e.g: "A map of ..."
 			MapValuesJson->RemoveField(TEXT("description"));
 
-			OutJson.Set(TEXT("type"), TEXT("object"));
-			OutJson.Set(TEXT("additionalProperties"), MapValuesJson);
+			OutputSchema->SetStringField(TEXT("type"), TEXT("object"));
+			OutputSchema->SetObjectField(TEXT("additionalProperties"), MapValuesJson);
 		}
 	}
 
-	static void BuildJsonForNumericProperty(const FNumericProperty& NumericProperty,
-		FJsonDomBuilder::FObject& OutJson)
+	static void BuildJsonForNumericProperty(const FNumericProperty& NumericProperty, const TSharedRef<FJsonObject>& OutputSchema)
 	{
 		if (const UEnum* EnumDef = NumericProperty.GetIntPropertyEnum())
 		{
-			BuildJsonForEnumDefinition(*EnumDef, OutJson);
+			BuildJsonForEnumDefinition(*EnumDef, OutputSchema);
 		}
 		else if (NumericProperty.IsFloatingPoint())
 		{
-			OutJson.Set(TEXT("type"), TEXT("number"));
+			OutputSchema->SetStringField(TEXT("type"), TEXT("number"));
 		}
 		else if (NumericProperty.IsInteger())
 		{
-			OutJson.Set(TEXT("type"), TEXT("integer"));
+			OutputSchema->SetStringField(TEXT("type"), TEXT("integer"));
 		}
 	}
 
-	static void BuildJsonForBoolProperty(FJsonDomBuilder::FObject& OutJson)
+	static void BuildJsonForBoolProperty(const TSharedRef<FJsonObject>& OutputSchema)
 	{
-		OutJson.Set(TEXT("type"), TEXT("boolean"));
+		OutputSchema->SetStringField(TEXT("type"), TEXT("boolean"));
 	}
 
-	static void BuildJsonForStringProperty(FJsonDomBuilder::FObject& OutJson)
+	static void BuildJsonForStringProperty(const TSharedRef<FJsonObject>& OutputSchema)
 	{
-		OutJson.Set(TEXT("type"), TEXT("string"));
+		OutputSchema->SetStringField(TEXT("type"), TEXT("string"));
 	}
 
+	FVisitorStack VisitorStack;
 	FJsonSchemaEditorMetadata& EditorMetadata; 
-	
-	FJsonSchemaPropertyFilter PropertyFilter;
 };
 
-//
-// Internal
-//
-
-namespace 
-{
-	template<class CharType, class PrintPolicy>
-	FString JsonObjectToJsonObjectString(const TSharedRef<FJsonObject>& JsonObject, const int32 Indent)
-	{
-		FString OutString;
-		
-		TSharedRef<TJsonWriter<CharType, PrintPolicy> > JsonWriter = TJsonWriterFactory<CharType, PrintPolicy>::Create(&OutString, Indent);
-		const bool bSuccess = FJsonSerializer::Serialize(JsonObject, JsonWriter);
-		JsonWriter->Close();
-		
-		return (bSuccess ? OutString : TEXT(""));
-	}
-}
-
 //
 // FJsonSchemaGenerator
 //
@@ -389,8 +448,7 @@ namespace
 //
 
 /*static*/ TSharedPtr<FJsonObject> FJsonSchemaGenerator::FPropertyToJsonSchemaObject(TNotNull<const FProperty*> Property,
-	const FJsonSchemaPropertyFilter& PropertyFilter, 
-	const FJsonSchemaEditorMetadata* CachedEditorMetadata)
+	const FJsonSchemaPropertyFilter& PropertyFilter, const FJsonSchemaEditorMetadata* CachedEditorMetadata)
 {
 	// See all NOTE_JSON_UTILITIES_SCHEMA_GENERATOR_METADATA_MODULAR_FEATURE.
 	FJsonSchemaEditorMetadata EditorMetadata;
@@ -400,43 +458,39 @@ namespace
 	}
 	else if (IJsonUtilitiesModularFeature* ModuleFeature = IJsonUtilitiesModularFeature::Get())
 	{
-		ModuleFeature->FPropertyToJsonSchemaMetadata(Property,
-			EditorMetadata,
-			PropertyFilter);
+		EditorMetadata = ModuleFeature->FPropertyToJsonSchemaMetadata(Property, PropertyFilter);
 	}
 	
-	UE_JSON_SCHEMA_SCOPED_MEMBER_PATH_PUSH(EditorMetadata.CurrentPropertyMemberPath, *Property->GetAuthoredName(), false);
-	
-	FJsonSchemaGeneratorVisitor JsonSchemaGeneratorVisitor(EditorMetadata, PropertyFilter);
-	
-	return JsonSchemaGeneratorVisitor.VisitFProperty(Property);
+	FJsonSchemaGeneratorVisitor JsonSchemaGeneratorVisitor(EditorMetadata);
+	return JsonSchemaGeneratorVisitor.VisitRoot(
+		FVisitorTarget(TInPlaceType<FValidConstFProperty>(), Property), 
+		FJsonSchemaMemberPath(*Property->GetAuthoredName()),\
+		PropertyFilter);
 }
 
 /*static*/ FString FJsonSchemaGenerator::FPropertyToJsonSchemaString(TNotNull<const FProperty*> Property,
-	const FJsonSchemaPropertyFilter& PropertyFilter,
-	const FJsonSchemaEditorMetadata* CachedEditorMetadata,
+	const FJsonSchemaPropertyFilter& PropertyFilter, const FJsonSchemaEditorMetadata* CachedEditorMetadata,
 	const int32 Indent, const bool bPrettyPrint)
 {
-	FString OutString;
+	FString OutputSchemaString;
 	
 	if (const TSharedPtr<FJsonObject> JsonSchema = FPropertyToJsonSchemaObject(Property, PropertyFilter, CachedEditorMetadata))
 	{
 		if (bPrettyPrint)
 		{
-			OutString = JsonObjectToJsonObjectString<TCHAR, TPrettyJsonPrintPolicy<TCHAR>>(JsonSchema.ToSharedRef(), Indent);
+			OutputSchemaString = JsonObjectToJsonObjectString<TCHAR, TPrettyJsonPrintPolicy<TCHAR>>(JsonSchema.ToSharedRef(), Indent);
 		}
 		else
 		{
-			OutString = JsonObjectToJsonObjectString<TCHAR, TCondensedJsonPrintPolicy<TCHAR>>(JsonSchema.ToSharedRef(), Indent);
+			OutputSchemaString = JsonObjectToJsonObjectString<TCHAR, TCondensedJsonPrintPolicy<TCHAR>>(JsonSchema.ToSharedRef(), Indent);
 		}
 	}
 
-	return OutString;
+	return OutputSchemaString;
 }
 
 /*static*/ TSharedPtr<FJsonObject> FJsonSchemaGenerator::UStructToJsonSchemaObject(TNotNull<const UStruct*> Struct,
-	const FJsonSchemaPropertyFilter& PropertyFilter,
-	const FJsonSchemaEditorMetadata* CachedEditorMetadata)
+	const FJsonSchemaPropertyFilter& PropertyFilter, const FJsonSchemaEditorMetadata* CachedEditorMetadata)
 {
 	// See all NOTE_JSON_UTILITIES_SCHEMA_GENERATOR_METADATA_MODULAR_FEATURE.
 	FJsonSchemaEditorMetadata EditorMetadata;
@@ -446,33 +500,33 @@ namespace
 	}
 	else if (IJsonUtilitiesModularFeature* ModuleFeature = IJsonUtilitiesModularFeature::Get())
 	{
-		ModuleFeature->UStructToJsonSchemaMetadata(Struct, 
-			EditorMetadata,
-			PropertyFilter);
+		EditorMetadata = ModuleFeature->UStructToJsonSchemaMetadata(Struct, PropertyFilter);
 	}
 	
-	FJsonSchemaGeneratorVisitor JsonSchemaGeneratorVisitor(EditorMetadata, PropertyFilter);
-	return JsonSchemaGeneratorVisitor.VisitUStruct(Struct, true);
+	FJsonSchemaGeneratorVisitor JsonSchemaGeneratorVisitor(EditorMetadata);
+	return JsonSchemaGeneratorVisitor.VisitRoot(
+		FVisitorTarget(TInPlaceType<FValidConstUStruct>(), Struct), 
+		FJsonSchemaMemberPath(),
+		PropertyFilter);
 }
 
 /*static*/ FString FJsonSchemaGenerator::UStructToJsonSchemaString(TNotNull<const UStruct*> Struct,
-	const FJsonSchemaPropertyFilter& PropertyFilter,
-	const FJsonSchemaEditorMetadata* CachedEditorMetadata,
+	const FJsonSchemaPropertyFilter& PropertyFilter, const FJsonSchemaEditorMetadata* CachedEditorMetadata,
 	const int32 Indent, const bool bPrettyPrint)
 {
-	FString OutString;
+	FString OutputSchemaString;
 	
 	if (const TSharedPtr<FJsonObject> JsonSchema = UStructToJsonSchemaObject(Struct, PropertyFilter, CachedEditorMetadata))
 	{
 		if (bPrettyPrint)
 		{
-			OutString = JsonObjectToJsonObjectString<TCHAR, TPrettyJsonPrintPolicy<TCHAR>>(JsonSchema.ToSharedRef(), Indent);
+			OutputSchemaString = JsonObjectToJsonObjectString<TCHAR, TPrettyJsonPrintPolicy<TCHAR>>(JsonSchema.ToSharedRef(), Indent);
 		}
 		else
 		{
-			OutString = JsonObjectToJsonObjectString<TCHAR, TCondensedJsonPrintPolicy<TCHAR>>(JsonSchema.ToSharedRef(), Indent);
+			OutputSchemaString = JsonObjectToJsonObjectString<TCHAR, TCondensedJsonPrintPolicy<TCHAR>>(JsonSchema.ToSharedRef(), Indent);
 		}
 	}
 
-	return OutString;
+	return OutputSchemaString;
 }
diff --git a/Engine/Source/Runtime/JsonUtilities/Private/JsonSchema/JsonSchemaVisitorStack.h b/Engine/Source/Runtime/JsonUtilities/Private/JsonSchema/JsonSchemaVisitorStack.h
new file mode 100644
--- /dev/null
+++ b/Engine/Source/Runtime/JsonUtilities/Private/JsonSchema/JsonSchemaVisitorStack.h
@@ -0,0 +1,58 @@
+// Copyright Epic Games, Inc. All Rights Reserved. 
+
+#pragma once
+
+#include "UObject/UnrealType.h"
+#include "Misc/NotNull.h"
+#include "Misc/TVariant.h"
+#include "Templates/SharedPointer.h"
+#include "Dom/JsonObject.h"
+#include "JsonSchema/JsonSchemaMemberPath.h"
+#include "JsonSchema/JsonSchemaPropertyFilter.h"
+
+namespace UE::JsonSchema
+{
+	using FValidConstUStruct = TNotNull<const UStruct*>;
+	using FValidConstFProperty = TNotNull<const FProperty*>;
+	using FVisitorTarget = TVariant<FValidConstUStruct, FValidConstFProperty>;
+	
+	enum class EVisitorStackElementFlags : uint8
+	{
+		None = 0x0,
+		IsRootStruct = 0x1,
+		IsFunctionWorldContextProperty = 0x2
+	};
+	ENUM_CLASS_FLAGS(EVisitorStackElementFlags);
+
+	struct FVisitorStackElement
+	{
+		explicit FVisitorStackElement(const FVisitorTarget& InVisitorTarget, 
+				const EVisitorStackElementFlags InFlags,
+				const FJsonSchemaMemberPath& InMemberPath, 
+				const FJsonSchemaPropertyFilter& InPropertyFilter,
+				const TSharedPtr<FJsonObject>& InOutputSchema) :
+			VisitorTarget(InVisitorTarget),
+			Flags(InFlags),
+			MemberPath(InMemberPath),
+			PropertyFilter(InPropertyFilter),
+			OutputSchema(InOutputSchema)
+		{
+			if (VisitorTarget.IsType<FValidConstUStruct>())
+			{
+				check(!EnumHasAnyFlags(Flags, EVisitorStackElementFlags::IsFunctionWorldContextProperty));
+			}
+			else if (VisitorTarget.IsType<FValidConstFProperty>())
+			{
+				check(!EnumHasAnyFlags(Flags, EVisitorStackElementFlags::IsRootStruct));
+			}
+		}
+		
+		FVisitorTarget VisitorTarget; 
+		EVisitorStackElementFlags Flags = EVisitorStackElementFlags::None;
+		FJsonSchemaMemberPath MemberPath;
+		FJsonSchemaPropertyFilter PropertyFilter;
+		const TSharedPtr<FJsonObject> OutputSchema;
+	};
+
+	using FVisitorStack = TArray<FVisitorStackElement>; 
+}
diff --git a/Engine/Source/Runtime/JsonUtilities/Private/JsonUtilitiesModularFeature.h b/Engine/Source/Runtime/JsonUtilities/Private/JsonUtilitiesModularFeature.h
--- a/Engine/Source/Runtime/JsonUtilities/Private/JsonUtilitiesModularFeature.h
+++ b/Engine/Source/Runtime/JsonUtilities/Private/JsonUtilitiesModularFeature.h
@@ -24,11 +24,9 @@ public:
 	JSONUTILITIES_API static IJsonUtilitiesModularFeature* Get();
 
 	
-	virtual void FPropertyToJsonSchemaMetadata(TNotNull<const FProperty*> Property,
-		FJsonSchemaEditorMetadata& OutEditorMetadata,
+	virtual FJsonSchemaEditorMetadata FPropertyToJsonSchemaMetadata(TNotNull<const FProperty*> Property,
 		const FJsonSchemaPropertyFilter& PropertyFilter = FJsonSchemaPropertyFilter()) const = 0;
 
-	virtual void UStructToJsonSchemaMetadata(TNotNull<const UStruct*> Struct,
-		FJsonSchemaEditorMetadata& OutEditorMetadata,
+	virtual FJsonSchemaEditorMetadata UStructToJsonSchemaMetadata(TNotNull<const UStruct*> Struct,
 		const FJsonSchemaPropertyFilter& PropertyFilter = FJsonSchemaPropertyFilter()) const = 0;
 };
diff --git a/Engine/Source/Runtime/JsonUtilities/Public/JsonSchema/JsonSchemaMemberPath.h b/Engine/Source/Runtime/JsonUtilities/Public/JsonSchema/JsonSchemaMemberPath.h
--- a/Engine/Source/Runtime/JsonUtilities/Public/JsonSchema/JsonSchemaMemberPath.h
+++ b/Engine/Source/Runtime/JsonUtilities/Public/JsonSchema/JsonSchemaMemberPath.h
@@ -20,7 +20,7 @@ public:
 	JSONUTILITIES_API FJsonSchemaMemberPath(const FJsonSchemaMemberPath& InMemberPath);
 	JSONUTILITIES_API explicit FJsonSchemaMemberPath(FStringView InitialString);
 	
-	FJsonSchemaMemberPath& operator=(const FJsonSchemaMemberPath& Other);
+	JSONUTILITIES_API FJsonSchemaMemberPath& operator=(const FJsonSchemaMemberPath& Other);
 	
 	/**
 	 * Add a string to the end.
