Json Utilities / Json Schema -

If UStruct is a UFunction, do not include "type": "object" in top-level schema.

#Synced-CL 49096081

diff --git a/Engine/Source/Runtime/JsonUtilities/Private/JsonSchema/JsonSchemaGenerator.cpp b/Engine/Source/Runtime/JsonUtilities/Private/JsonSchema/JsonSchemaGenerator.cpp
--- a/Engine/Source/Runtime/JsonUtilities/Private/JsonSchema/JsonSchemaGenerator.cpp
+++ b/Engine/Source/Runtime/JsonUtilities/Private/JsonSchema/JsonSchemaGenerator.cpp
@@ -149,8 +149,14 @@ public:
 		const EVisitorStackElementFlags Flags)
 	{
 		EditorMetadata.CurrentPropertyMemberPath = MemberPath;
-		
-		OutputSchema->SetStringField(TEXT("type"), TEXT("object"));
+
+		const bool bIsUFunction = Struct->IsA(UFunction::StaticClass());
+
+		// Only include "type": "object" in this top-level schema if it's not a function, and so not following JSON RPC format.
+		if (!bIsUFunction)
+		{
+			OutputSchema->SetStringField(TEXT("type"), TEXT("object"));
+		}
 
 		// Special case handling for FJsonObjectWrapper, where we know this must be a Json object.   
 		if (Struct == FJsonObjectWrapper::StaticStruct())
@@ -168,10 +174,9 @@ public:
 			}
 		}
 
-		// Adds property to JSON and possibly the required list.
-		// If added, property will also be added to the visitor stack for processing.
-		// Returns whether the property was added to the JSON object.
-		auto MaybeAddProperty = [this, PropertyFilter, Flags](const FProperty* Property,
+		// If property is to be included - Adds property to JSON and possibly the required list, and also adds to the visitor stack for processing.
+		// Returns whether the property was included.
+		auto MaybeIncludeProperty = [this, PropertyFilter, Flags](const FProperty* Property,
 			const TSharedRef<FJsonObject>& PropertiesSchema, TArray<TSharedPtr<FJsonValue>>& RequiredPropertiesArray) -> bool
 			{
 				const FString PropertyName = Property->GetAuthoredName();
@@ -222,7 +227,7 @@ public:
 		// Properties & required.
 		// If this struct is a function, use JSON RPC format (with inputSchema/outputSchema). 
 		// Otherwise, use standard format.
-		if (Struct->IsA(UFunction::StaticClass()))
+		if (bIsUFunction)
 		{
 			// If struct is a function, output JSON RPC -
 			// [REDACTED_PROFILE_URL]/modelcontextprotocol/blob/main/schema/2025-06-18/schema.ts#L928
@@ -241,11 +246,11 @@ public:
 			{
 				if (Property->HasAnyPropertyFlags(CPF_ReturnParm))
 				{
-					MaybeAddProperty(Property, FunctionOutputPropertiesSchema, FunctionOutputRequiredPropertiesArray);
+					MaybeIncludeProperty(Property, FunctionOutputPropertiesSchema, FunctionOutputRequiredPropertiesArray);
 				}
 				else
 				{
-					MaybeAddProperty(Property, FunctionInputPropertiesSchema, FunctionInputRequiredPropertiesArray);
+					MaybeIncludeProperty(Property, FunctionInputPropertiesSchema, FunctionInputRequiredPropertiesArray);
 				}
 			}
 	
@@ -292,7 +297,7 @@ public:
 			TArray<TSharedPtr<FJsonValue>> RequiredPropertiesArray;
 			for (const FProperty* Property : TFieldRange<FProperty>(Struct))
 			{
-				MaybeAddProperty(Property, PropertiesSchema, RequiredPropertiesArray);
+				MaybeIncludeProperty(Property, PropertiesSchema, RequiredPropertiesArray);
 			}
 	
 			if (PropertiesSchema->Values.Num() > 0)
