Json Utilities / Json Schema

- Changed "defaultValue" to "default", to properly adhere to schema standard.

- Default values are not always given as a string with quotes, and will be without quotes for numeric and bool, while string and enum properties do have quotes.

- Trivial: Renamed a variable to clarify when to output JSON RPC format for UFunctions.

#Synced-CL 49315116

diff --git a/Engine/Source/Editor/JsonUtilitiesEditor/Private/JsonSchema/JsonSchemaGeneratorEditor.cpp b/Engine/Source/Editor/JsonUtilitiesEditor/Private/JsonSchema/JsonSchemaGeneratorEditor.cpp
--- a/Engine/Source/Editor/JsonUtilitiesEditor/Private/JsonSchema/JsonSchemaGeneratorEditor.cpp
+++ b/Engine/Source/Editor/JsonUtilitiesEditor/Private/JsonSchema/JsonSchemaGeneratorEditor.cpp
@@ -87,6 +87,7 @@ public:
 			}
 		}
 
+		// Add minimum/maximum for numeric properties.
 		if (const FNumericProperty* NumericProperty = CastField<FNumericProperty>(Property))
 		{
 			if (NumericProperty->HasMetaData("ClampMin"))
@@ -99,13 +100,28 @@ public:
 			}
 		}
 
+		// If function, add any parameter defaults.
 		if (PropertyOwnerFunction)
 		{
 			FString ParameterDefault;
 			UEdGraphSchema_K2::FindFunctionParameterDefaultValue(PropertyOwnerFunction, Property, ParameterDefault);
 			if (!ParameterDefault.IsEmpty())
 			{
-				PropertyMetadata->SetStringField(TEXT("defaultValue"), ParameterDefault);
+				if (Property->IsA<FNumericProperty>())
+				{
+					// Convert default string to double, ints will end up being rounded in the JSON.
+					PropertyMetadata->SetNumberField(TEXT("default"), FCString::Atod(*ParameterDefault));
+				}
+				else if (Property->IsA<FBoolProperty>())
+				{
+					// Default string will be either 'true' or 'false' only.
+					PropertyMetadata->SetBoolField(TEXT("default"), ParameterDefault.Equals("true"));
+				}
+				else
+				{
+					// Default string works for both strings and enums.
+					PropertyMetadata->SetStringField(TEXT("default"), ParameterDefault);
+				}
 			}
 		}
 
diff --git a/Engine/Source/Runtime/JsonUtilities/Private/JsonSchema/JsonSchemaGenerator.cpp b/Engine/Source/Runtime/JsonUtilities/Private/JsonSchema/JsonSchemaGenerator.cpp
--- a/Engine/Source/Runtime/JsonUtilities/Private/JsonSchema/JsonSchemaGenerator.cpp
+++ b/Engine/Source/Runtime/JsonUtilities/Private/JsonSchema/JsonSchemaGenerator.cpp
@@ -150,10 +150,10 @@ public:
 	{
 		EditorMetadata.CurrentPropertyMemberPath = MemberPath;
 
-		const bool bIsUFunction = Struct->IsA(UFunction::StaticClass());
+		const bool bUseJsonRpcFormat = Struct->IsA(UFunction::StaticClass());
 
 		// Only include "type": "object" in this top-level schema if it's not a function, and so not following JSON RPC format.
-		if (!bIsUFunction)
+		if (!bUseJsonRpcFormat)
 		{
 			OutputSchema->SetStringField(TEXT("type"), TEXT("object"));
 		}
@@ -213,7 +213,7 @@ public:
 					// Property has been explicitly required.
 					RequiredPropertiesArray.Add(MakeShared<FJsonValueString>(JsonMemberName));
 				}
-				else if (!(PropertyMetadata && PropertyMetadata->HasField(TEXT("defaultValue"))) &&
+				else if (!(PropertyMetadata && PropertyMetadata->HasField(TEXT("default"))) &&
 					!Property->IsA<FOptionalProperty>())
 				{
 					// Property has not been explicitly required.
@@ -227,7 +227,7 @@ public:
 		// Properties & required.
 		// If this struct is a function, use JSON RPC format (with inputSchema/outputSchema). 
 		// Otherwise, use standard format.
-		if (bIsUFunction)
+		if (bUseJsonRpcFormat)
 		{
 			// If struct is a function, output JSON RPC -
 			// [REDACTED_PROFILE_URL]/modelcontextprotocol/blob/main/schema/2025-06-18/schema.ts#L928
