PR #12417: FJsonObjectConverter: Allow different formats when deserializing FGuid properties

#12417

diff --git a/Engine/Source/Runtime/JsonUtilities/Private/JsonObjectConverter.cpp b/Engine/Source/Runtime/JsonUtilities/Private/JsonObjectConverter.cpp
--- a/Engine/Source/Runtime/JsonUtilities/Private/JsonObjectConverter.cpp
+++ b/Engine/Source/Runtime/JsonUtilities/Private/JsonObjectConverter.cpp
@@ -45,6 +45,7 @@ namespace
 	const FString ObjectClassNameKey = "_ClassName";
 
 	const FName NAME_DateTime("DateTime");
+	const FName NAME_Guid("Guid");
 
 	TSharedPtr<FJsonValue> FPropertyToJsonValueWithContainer(FProperty* Property, const void* Value, const UObject* Container, TSet<const UObject*>* ExportedObjects, int64 CheckFlags, int64 SkipFlags, const FJsonObjectConverter::CustomExportCallback* ExportCb, FProperty* OuterProperty, EJsonObjectConversionFlags ConversionFlags);
 	bool UStructToJsonAttributesWithContainer(const UStruct* StructDefinition, const void* Struct, const UObject* Container, TSet<const UObject*>* ExportedObjects, TMap< FString, TSharedPtr<FJsonValue> >& OutJsonAttributes, int64 CheckFlags, int64 SkipFlags, const FJsonObjectConverter::CustomExportCallback* ExportCb, EJsonObjectConversionFlags ConversionFlags);
@@ -922,6 +923,29 @@ namespace
 					return false;
 				}
 			}
+			else if (JsonValue->Type == EJson::String && StructProperty->Struct->GetFName() == NAME_Guid)
+			{
+				FString GuidString = JsonValue->AsString();
+				FGuid& GuidOut = *(FGuid*)OutValue;
+				if (FGuid::Parse(GuidString, GuidOut))
+				{
+					// ok
+				}
+				// Try again with regular ImportText if FGuid::Parse fails
+				else if (Property->ImportText_Direct(*GuidString, OutValue, nullptr, PPF_None))
+				{
+					// ok
+				}
+				else
+				{
+					UE_LOG(LogJson, Error, TEXT("JsonValueToUProperty - Unable to import JSON string into Guid property %s"), *Property->GetAuthoredName());
+					if (OutFailReason)
+					{
+						*OutFailReason = FText::Format(LOCTEXT("FailImportGuidFromString", "Unable to import JSON string into Guid property {0}"), FText::FromString(Property->GetAuthoredName()));
+					}
+					return false;
+				}
+			}
 			else if (JsonValue->Type == EJson::String && StructProperty->Struct->GetCppStructOps() && StructProperty->Struct->GetCppStructOps()->HasImportTextItem())
 			{
 				UScriptStruct::ICppStructOps* TheCppStructOps = StructProperty->Struct->GetCppStructOps();
