PR #14253: [Json] Enforce the class of a json object being deserialized

#14253

diff --git a/Engine/Source/Runtime/JsonUtilities/Private/JsonObjectConverter.cpp b/Engine/Source/Runtime/JsonUtilities/Private/JsonObjectConverter.cpp
--- a/Engine/Source/Runtime/JsonUtilities/Private/JsonObjectConverter.cpp
+++ b/Engine/Source/Runtime/JsonUtilities/Private/JsonObjectConverter.cpp
@@ -1002,6 +1002,13 @@ namespace
 
 				TSharedPtr<FJsonObject> Obj = JsonValue->AsObject();
 				UClass* PropertyClass = ObjectProperty->PropertyClass;
+				
+				if (!ObjectProperty->HasAnyPropertyFlags(CPF_InstancedReference | CPF_PersistentInstance | CPF_ContainsInstancedReference) && !PropertyClass->HasAnyClassFlags(CLASS_DefaultToInstanced))
+				{
+					UE_LOG(LogJson, Error, TEXT("JsonValueToUProperty - Unable to import JSON object into property %s (property is not instanced)"), *Property->GetAuthoredName());
+					*OutFailReason = FText::Format(LOCTEXT("FailImportObjectNotInstanced", "Unable to import JSON object into property {0} (property is not instanced)"), FText::FromString(Property->GetAuthoredName()));
+					return false;
+				}
 
 				// If a specific subclass was stored in the JSON, use that instead of the PropertyClass
 				FString ClassString;
@@ -1013,6 +1020,17 @@ namespace
 						UClass* FoundClass = FPackageName::IsShortPackageName(ClassString) ? FindFirstObject<UClass>(*ClassString) : LoadClass<UObject>(nullptr, *ClassString);
 						if (FoundClass)
 						{
+							if (!FoundClass->IsChildOf(PropertyClass) || FoundClass->HasAnyClassFlags(CLASS_Abstract))
+							{
+								UE_LOG(LogJson, Error, TEXT("JsonValueToUProperty - Unable to import JSON object of type %s into property %s of type %s"), *GetNameSafe(FoundClass), *Property->GetAuthoredName(), *GetNameSafe(PropertyClass));
+								*OutFailReason = FText::Format(
+									LOCTEXT("FailImportObjectClass", "Unable to import JSON object of type {0} into property {1} of type {2}"),
+									FText::FromString(GetNameSafe(FoundClass)),
+									FText::FromString(Property->GetAuthoredName()),
+									FText::FromString(GetNameSafe(PropertyClass)));
+								return false;
+							}
+
 							PropertyClass = FoundClass;
 						}
 						else
