JsonUtilties / JsonSchema -

BUGFIX in FJsonSchemaMemberPath -
Fix for "AddressSanitizer: heap-use-after-free" crash (reported by [REDACTED_PERSON] / [REDACTED_PERSON].)
Forcing FJsonSchemaMemberPath (TStringBuilder) to be mutable, with a const cast and dereference, to avoid compiler warnings - this ultimately calls an 'ensure null terminated' function which writes into its const buffer.
Instead, using safe ToView(), which never attempts a write operation into the buffer.

#Synced-CL 50173404

diff --git a/Engine/Plugins/Tests/EditorTests/Source/EditorTests/Private/UnrealEd/JsonUtilitiesTests.cpp b/Engine/Plugins/Tests/EditorTests/Source/EditorTests/Private/UnrealEd/JsonUtilitiesTests.cpp
--- a/Engine/Plugins/Tests/EditorTests/Source/EditorTests/Private/UnrealEd/JsonUtilitiesTests.cpp
+++ b/Engine/Plugins/Tests/EditorTests/Source/EditorTests/Private/UnrealEd/JsonUtilitiesTests.cpp
@@ -55,9 +55,9 @@ PRAGMA_DISABLE_DEPRECATION_WARNINGS
 namespace
 {
 	// Gets last member of the member path.
-	FString Last(FJsonSchemaMemberPath& MemberPath)
+	FString Last(const FJsonSchemaMemberPath& MemberPath)
 	{
-		const FString MemberPathString = MemberPath.ToString();
+		const FString MemberPathString(MemberPath.ToView());
 	
 		if (int32 Index = INDEX_NONE;
 			MemberPathString.FindLastChar(FJsonSchemaMemberPath::Delimiter[0], Index))
diff --git a/Engine/Source/Runtime/JsonUtilities/Private/JsonSchema/JsonSchemaMemberPath.cpp b/Engine/Source/Runtime/JsonUtilities/Private/JsonSchema/JsonSchemaMemberPath.cpp
--- a/Engine/Source/Runtime/JsonUtilities/Private/JsonSchema/JsonSchemaMemberPath.cpp
+++ b/Engine/Source/Runtime/JsonUtilities/Private/JsonSchema/JsonSchemaMemberPath.cpp
@@ -27,9 +27,7 @@ FJsonSchemaMemberPath& FJsonSchemaMemberPath::operator=(const FJsonSchemaMemberP
 	if (this != &Other)
 	{
 		Reset();
-		
-		FJsonSchemaMemberPath& InMemberPathMutable = *const_cast<FJsonSchemaMemberPath*>(&Other); 
-		Append(*InMemberPathMutable);
+		Append(Other.ToView());
 	}
 	
 	return *this;
@@ -95,6 +93,5 @@ bool FJsonSchemaMemberPath::PopSubmember(FStringView SubmemberString)
 
 FText FJsonSchemaMemberPath::ToText() const
 {
-	FJsonSchemaMemberPath& ThisMutable = *const_cast<FJsonSchemaMemberPath*>(this); 
-	return FText::FromString(*ThisMutable);
+	return FText::FromString(FString(ToView()));
 }
