Fixed JSON import errors writing to a null OutFailReason.

diff --git a/Engine/Source/Runtime/JsonUtilities/Private/JsonObjectConverter.cpp b/Engine/Source/Runtime/JsonUtilities/Private/JsonObjectConverter.cpp
--- a/Engine/Source/Runtime/JsonUtilities/Private/JsonObjectConverter.cpp
+++ b/Engine/Source/Runtime/JsonUtilities/Private/JsonObjectConverter.cpp
@@ -1006,7 +1006,10 @@ namespace
 				if (!ObjectProperty->HasAnyPropertyFlags(CPF_InstancedReference | CPF_PersistentInstance | CPF_ContainsInstancedReference) && !PropertyClass->HasAnyClassFlags(CLASS_DefaultToInstanced))
 				{
 					UE_LOG(LogJson, Error, TEXT("JsonValueToUProperty - Unable to import JSON object into property %s (property is not instanced)"), *Property->GetAuthoredName());
-					*OutFailReason = FText::Format(LOCTEXT("FailImportObjectNotInstanced", "Unable to import JSON object into property {0} (property is not instanced)"), FText::FromString(Property->GetAuthoredName()));
+					if (OutFailReason)
+					{
+						*OutFailReason = FText::Format(LOCTEXT("FailImportObjectNotInstanced", "Unable to import JSON object into property {0} (property is not instanced)"), FText::FromString(Property->GetAuthoredName()));
+					}
 					return false;
 				}
 
@@ -1023,11 +1026,14 @@ namespace
 							if (!FoundClass->IsChildOf(PropertyClass) || FoundClass->HasAnyClassFlags(CLASS_Abstract))
 							{
 								UE_LOG(LogJson, Error, TEXT("JsonValueToUProperty - Unable to import JSON object of type %s into property %s of type %s"), *GetNameSafe(FoundClass), *Property->GetAuthoredName(), *GetNameSafe(PropertyClass));
-								*OutFailReason = FText::Format(
-									LOCTEXT("FailImportObjectClass", "Unable to import JSON object of type {0} into property {1} of type {2}"),
-									FText::FromString(GetNameSafe(FoundClass)),
-									FText::FromString(Property->GetAuthoredName()),
-									FText::FromString(GetNameSafe(PropertyClass)));
+								if (OutFailReason)
+								{
+									*OutFailReason = FText::Format(
+										LOCTEXT("FailImportObjectClass", "Unable to import JSON object of type {0} into property {1} of type {2}"),
+										FText::FromString(GetNameSafe(FoundClass)),
+										FText::FromString(Property->GetAuthoredName()),
+										FText::FromString(GetNameSafe(PropertyClass)));
+								}
 								return false;
 							}
 
