Fix heap-use-after-free in FJsonSchemaMetadataVisitor when generating schemas for UFunctions with many parameters.

FJsonSchemaMemberPath inherited from TStringBuilder<512> which stores a self-referential pointer into its own inline buffer; TArray's realloc-based growth would memcpy these elements without fixing up the pointer, leaving it dangling. Replaced the TStringBuilder<512> base class with an FString member, which is safely relocatable.

#Synced-CL 50702752

diff --git a/Engine/Source/Runtime/JsonUtilities/Private/JsonSchema/JsonSchemaEditorMetadata.cpp b/Engine/Source/Runtime/JsonUtilities/Private/JsonSchema/JsonSchemaEditorMetadata.cpp
--- a/Engine/Source/Runtime/JsonUtilities/Private/JsonSchema/JsonSchemaEditorMetadata.cpp
+++ b/Engine/Source/Runtime/JsonUtilities/Private/JsonSchema/JsonSchemaEditorMetadata.cpp
@@ -19,8 +19,7 @@ FJsonSchemaEditorMetadata& FJsonSchemaEditorMetadata::operator=(const FJsonSchem
 
 		PropertyMemberPathToPropertyMetadataMap = Other.PropertyMemberPathToPropertyMetadataMap;
 
-		CurrentPropertyMemberPath.Reset();
-		CurrentPropertyMemberPath.Append(*const_cast<FJsonSchemaMemberPath&>(Other.CurrentPropertyMemberPath));
+		CurrentPropertyMemberPath = Other.CurrentPropertyMemberPath;
 	}
 	
 	return *this;
diff --git a/Engine/Source/Runtime/JsonUtilities/Private/JsonSchema/JsonSchemaMemberPath.cpp b/Engine/Source/Runtime/JsonUtilities/Private/JsonSchema/JsonSchemaMemberPath.cpp
--- a/Engine/Source/Runtime/JsonUtilities/Private/JsonSchema/JsonSchemaMemberPath.cpp
+++ b/Engine/Source/Runtime/JsonUtilities/Private/JsonSchema/JsonSchemaMemberPath.cpp
@@ -1,4 +1,4 @@
-// Copyright Epic Games, Inc. All Rights Reserved. 
+// Copyright Epic Games, Inc. All Rights Reserved.
 
 #include "JsonSchema/JsonSchemaMemberPath.h"
 
@@ -13,23 +13,22 @@ FJsonSchemaMemberPath::FJsonSchemaMemberPath()
 }
 
 FJsonSchemaMemberPath::FJsonSchemaMemberPath(const FJsonSchemaMemberPath& InMemberPath)
+	: PathString(InMemberPath.PathString)
 {
-	*this = InMemberPath;
 }
 
 FJsonSchemaMemberPath::FJsonSchemaMemberPath(FStringView InitialString)
+	: PathString(InitialString)
 {
-	Append(InitialString);
 }
 
 FJsonSchemaMemberPath& FJsonSchemaMemberPath::operator=(const FJsonSchemaMemberPath& Other)
 {
 	if (this != &Other)
 	{
-		Reset();
-		Append(Other.ToView());
+		PathString = Other.PathString;
 	}
-	
+
 	return *this;
 }
 
@@ -39,8 +38,8 @@ bool FJsonSchemaMemberPath::Push(FStringView String)
 	{
 		return false;
 	}
-	
-	Append(String);
+
+	PathString.Append(String);
 
 	return true;
 }
@@ -51,7 +50,7 @@ bool FJsonSchemaMemberPath::Pop(FStringView String)
 	{
 		return false;
 	}
-	
+
 	// Safety check. Make sure current property path indeed ends with the text the caller wants to trim from the end.
 	{
 		const int32 Offset = Len() - String.Len();
@@ -59,15 +58,15 @@ bool FJsonSchemaMemberPath::Pop(FStringView String)
 		{
 			return false;
 		}
-		
-		if (!ensureAlways(String.Equals(ToString() + Offset, ESearchCase::CaseSensitive)))
+
+		if (!ensureAlways(String.Equals(GetData() + Offset, ESearchCase::CaseSensitive)))
 		{
 			return false;
 		}
 	}
-	
-	RemoveSuffix(String.Len());
-	
+
+	PathString.LeftInline(Len() - String.Len());
+
 	return true;
 }
 
@@ -77,7 +76,7 @@ bool FJsonSchemaMemberPath::PushSubmember(FStringView SubmemberString)
 	{
 		return false;
 	}
-	
+
 	return (Push(Delimiter) && Push(SubmemberString));
 }
 
@@ -87,11 +86,11 @@ bool FJsonSchemaMemberPath::PopSubmember(FStringView SubmemberString)
 	{
 		return false;
 	}
-	
+
 	return (Pop(SubmemberString) && Pop(Delimiter));
 }
 
 FText FJsonSchemaMemberPath::ToText() const
 {
-	return FText::FromString(FString(ToView()));
+	return FText::FromString(PathString);
 }
diff --git a/Engine/Source/Runtime/JsonUtilities/Public/JsonSchema/JsonSchemaMemberPath.h b/Engine/Source/Runtime/JsonUtilities/Public/JsonSchema/JsonSchemaMemberPath.h
--- a/Engine/Source/Runtime/JsonUtilities/Public/JsonSchema/JsonSchemaMemberPath.h
+++ b/Engine/Source/Runtime/JsonUtilities/Public/JsonSchema/JsonSchemaMemberPath.h
@@ -7,7 +7,7 @@
 /**
  * JSON member path builder.
  */
-class FJsonSchemaMemberPath : public TStringBuilder<512>
+class FJsonSchemaMemberPath
 {
 public:
 
@@ -18,17 +18,33 @@ public:
 
 	JSONUTILITIES_API FJsonSchemaMemberPath();
 	JSONUTILITIES_API FJsonSchemaMemberPath(const FJsonSchemaMemberPath& InMemberPath);
+	FJsonSchemaMemberPath(FJsonSchemaMemberPath&& Other) = default;
 	JSONUTILITIES_API explicit FJsonSchemaMemberPath(FStringView InitialString);
-	
+
 	JSONUTILITIES_API FJsonSchemaMemberPath& operator=(const FJsonSchemaMemberPath& Other);
-	
+	FJsonSchemaMemberPath& operator=(FJsonSchemaMemberPath&& Other) = default;
+
+	/** Read-only subset of TStringBuilderBase's interface, for compatibility with existing callers. */
+
+	/** Returns the length of the path string. */
+	int32 Len() const { return PathString.Len(); }
+
+	/** Returns a pointer to the path string data. */
+	const TCHAR* GetData() const { return *PathString; }
+
+	/** Returns a null-terminated TCHAR pointer to the path string. */
+	const TCHAR* operator*() const { return *PathString; }
+
+	/** Returns a string view of the path. */
+	FStringView ToView() const { return FStringView(PathString); }
+
 	/**
 	 * Add a string to the end.
 	 * @param String String add to end.
 	 * @return Success.
 	 */
 	JSONUTILITIES_API bool Push(FStringView String);
-	
+
 	/**
 	 * Remove a string from the end.
 	 * @param String String to remove from end.
@@ -42,19 +58,23 @@ public:
 	 * @return Success.
 	 */
 	JSONUTILITIES_API bool PushSubmember(FStringView SubmemberString);
-	
+
 	/**
 	 * Remove a string from the end, as a submember: Delimiter + SubmemberText
 	 * @param SubmemberString String to remove from end.
 	 * @return Success.
 	 */
 	JSONUTILITIES_API bool PopSubmember(FStringView SubmemberString);
-	
+
 	/**
 	 * Converts path to FText.
 	 * @return FText representation of path.
 	 */
-	 FText ToText() const;
+	FText ToText() const;
+
+private:
+
+	FString PathString;
 };
 
 /**
