Add support for PropertyBag and InstancedStruct in JsonObjectConverter and JsonSchemaGenerator.

#Synced-CL 51698457

diff --git a/Engine/Plugins/Tests/EditorTests/Source/EditorTests/Private/UnrealEd/JsonUtilitiesTests.cpp b/Engine/Plugins/Tests/EditorTests/Source/EditorTests/Private/UnrealEd/JsonUtilitiesTests.cpp
--- a/Engine/Plugins/Tests/EditorTests/Source/EditorTests/Private/UnrealEd/JsonUtilitiesTests.cpp
+++ b/Engine/Plugins/Tests/EditorTests/Source/EditorTests/Private/UnrealEd/JsonUtilitiesTests.cpp
@@ -9,7 +9,9 @@
 #include "Tests/AutomationEditorCommon.h"
 #endif
 
+#include "JsonObjectConverter.h"
 #include "JsonSchema/JsonSchemaGenerator.h"
+#include "StructUtils/PropertyBag.h"
 
 
 // Expected values.
@@ -1075,6 +1077,219 @@ void FJsonUtilitiesEditorSpec::Define()
 					TestFalse("Schema is invalid.", Schema.IsValid());					
 			});
 
+			It(
+				"Should generate a fallback FInstancedStruct schema without instance memory.",
+				[this]() -> void
+				{
+					UStruct* WrapperStruct = FJsonUtilitiesInstancedStructWrapper::StaticStruct();
+					FProperty* Prop = WrapperStruct->FindPropertyByName("InstancedStructField");
+					const TSharedPtr<FJsonObject> Schema = FJsonSchemaGenerator::FPropertyToJsonSchemaObject(Prop);
+					TestBasics(*this, Schema, "object", "", FString("FInstancedStruct"));
+					TestFalse("No properties field", Schema->HasField(TEXT("properties")));
+					TestFalse("No required field", Schema->HasField(TEXT("required")));
+				});
+
+			It(
+				"Should generate full FInstancedStruct schema when instance memory is provided.",
+				[this]() -> void
+				{
+					UStruct* WrapperStruct = FJsonUtilitiesInstancedStructWrapper::StaticStruct();
+					FProperty* Prop = WrapperStruct->FindPropertyByName("InstancedStructField");
+					FInstancedStruct InstancedStruct;
+					InstancedStruct.InitializeAs<FJsonUtilitiesInstancedInnerStruct>();
+					const TSharedPtr<FJsonObject> Schema = FJsonSchemaGenerator::FPropertyToJsonSchemaObject(
+						Prop, FJsonSchemaPropertyFilter(), nullptr, &InstancedStruct);
+					TestBasics(*this, Schema, "object", "", FString("JsonUtilitiesInstancedInnerStruct"));
+					const TSharedPtr<FJsonObject>* Properties = nullptr;
+					TestTrue("Properties", Schema->TryGetObjectField(TEXT("properties"), Properties));
+					if (Properties)
+					{
+						const TSharedPtr<FJsonObject>* FloatField = nullptr;
+						TestTrue("FloatValue", Properties->Get()->TryGetObjectField(TEXT("FloatValue"), FloatField));
+						if (FloatField)
+						{
+							TestFloat(*this, *FloatField, "");
+						}
+						const TSharedPtr<FJsonObject>* StringField = nullptr;
+						TestTrue("StringValue", Properties->Get()->TryGetObjectField(TEXT("StringValue"), StringField));
+						if (StringField)
+						{
+							TestString(*this, *StringField, "");
+						}
+					}
+					TestRequired(*this, Schema, TArray<FString>({"FloatValue", "StringValue"}));
+				});
+
+			It(
+				"Should generate a FInstancedPropertyBag fallback schema without instance memory.",
+				[this]() -> void
+				{
+					UStruct* WrapperStruct = FJsonUtilitiesPropertyBagWrapper::StaticStruct();
+					FProperty* Prop = WrapperStruct->FindPropertyByName("PropertyBagField");
+					const TSharedPtr<FJsonObject> Schema = FJsonSchemaGenerator::FPropertyToJsonSchemaObject(Prop);
+					TestBasics(*this, Schema, "object", "", FString("FInstancedPropertyBag"));
+					TestFalse("No properties field", Schema->HasField(TEXT("properties")));
+				});
+
+			It(
+				"Should generate full FInstancedPropertyBag schema when instance memory is provided.",
+				[this]() -> void
+				{
+					UStruct* WrapperStruct = FJsonUtilitiesPropertyBagWrapper::StaticStruct();
+					FProperty* Prop = WrapperStruct->FindPropertyByName("PropertyBagField");
+					FInstancedPropertyBag Bag;
+					Bag.AddProperty(FName("TestFloat"), EPropertyBagPropertyType::Float);
+					Bag.AddProperty(FName("TestString"), EPropertyBagPropertyType::String);
+					const TSharedPtr<FJsonObject> Schema = FJsonSchemaGenerator::FPropertyToJsonSchemaObject(
+						Prop, FJsonSchemaPropertyFilter(), nullptr, &Bag);
+					TestTrue("Schema", Schema.IsValid());
+					if (!Schema.IsValid())
+					{
+						return;
+					}
+					FString Type;
+					TestTrue("Type field present", Schema->TryGetStringField(TEXT("type"), Type));
+					TestEqual("Type", Type, FString("object"));
+					const TSharedPtr<FJsonObject>* Properties = nullptr;
+					TestTrue("Properties", Schema->TryGetObjectField(TEXT("properties"), Properties));
+					if (Properties)
+					{
+						const TSharedPtr<FJsonObject>* FloatProp = nullptr;
+						TestTrue("TestFloat", Properties->Get()->TryGetObjectField(TEXT("TestFloat"), FloatProp));
+						if (FloatProp)
+						{
+							TestFloat(*this, *FloatProp, "");
+						}
+						const TSharedPtr<FJsonObject>* StringProp = nullptr;
+						TestTrue("TestString", Properties->Get()->TryGetObjectField(TEXT("TestString"), StringProp));
+						if (StringProp)
+						{
+							TestString(*this, *StringProp, "");
+						}
+					}
+				});
+
+			It(
+				"Should generate a fallback TOptional<FInstancedStruct> schema without instance memory.",
+				[this]() -> void
+				{
+					UStruct* WrapperStruct = FJsonUtilitiesInstancedStructWrapper::StaticStruct();
+					FProperty* Prop = WrapperStruct->FindPropertyByName("OptionalInstancedStructField");
+					const TSharedPtr<FJsonObject> Schema = FJsonSchemaGenerator::FPropertyToJsonSchemaObject(Prop);
+					TestBasics(*this, Schema, "object", "", FString("FInstancedStruct"));
+					TestFalse("No properties field", Schema->HasField(TEXT("properties")));
+					TestFalse("No required field", Schema->HasField(TEXT("required")));
+				});
+
+			It(
+				"Should generate full TOptional<FInstancedStruct> schema when instance memory is provided and the optional is set.",
+				[this]() -> void
+				{
+					UStruct* WrapperStruct = FJsonUtilitiesInstancedStructWrapper::StaticStruct();
+					FProperty* Prop = WrapperStruct->FindPropertyByName("OptionalInstancedStructField");
+					TOptional<FInstancedStruct> OptionalField;
+					OptionalField.Emplace();
+					OptionalField->InitializeAs<FJsonUtilitiesInstancedInnerStruct>();
+					const TSharedPtr<FJsonObject> Schema = FJsonSchemaGenerator::FPropertyToJsonSchemaObject(
+						Prop, FJsonSchemaPropertyFilter(), nullptr, &OptionalField);
+					TestBasics(*this, Schema, "object", "", FString("JsonUtilitiesInstancedInnerStruct"));
+					const TSharedPtr<FJsonObject>* Properties = nullptr;
+					TestTrue("Properties", Schema->TryGetObjectField(TEXT("properties"), Properties));
+					if (Properties)
+					{
+						const TSharedPtr<FJsonObject>* FloatField = nullptr;
+						TestTrue("FloatValue", Properties->Get()->TryGetObjectField(TEXT("FloatValue"), FloatField));
+						if (FloatField)
+						{
+							TestFloat(*this, *FloatField, "");
+						}
+						const TSharedPtr<FJsonObject>* StringField = nullptr;
+						TestTrue("StringValue", Properties->Get()->TryGetObjectField(TEXT("StringValue"), StringField));
+						if (StringField)
+						{
+							TestString(*this, *StringField, "");
+						}
+					}
+					TestRequired(*this, Schema, TArray<FString>({"FloatValue", "StringValue"}));
+				});
+
+			It(
+				"Should generate a fallback TOptional<FInstancedStruct> schema when instance memory is provided but the optional is not set.",
+				[this]() -> void
+				{
+					UStruct* WrapperStruct = FJsonUtilitiesInstancedStructWrapper::StaticStruct();
+					FProperty* Prop = WrapperStruct->FindPropertyByName("OptionalInstancedStructField");
+					TOptional<FInstancedStruct> OptionalField; // not set
+					const TSharedPtr<FJsonObject> Schema = FJsonSchemaGenerator::FPropertyToJsonSchemaObject(
+						Prop, FJsonSchemaPropertyFilter(), nullptr, &OptionalField);
+					TestBasics(*this, Schema, "object", "", FString("FInstancedStruct"));
+					TestFalse("No properties field", Schema->HasField(TEXT("properties")));
+					TestFalse("No required field", Schema->HasField(TEXT("required")));
+				});
+
+			It(
+				"Should generate a fallback TOptional<FInstancedPropertyBag> schema without instance memory.",
+				[this]() -> void
+				{
+					UStruct* WrapperStruct = FJsonUtilitiesPropertyBagWrapper::StaticStruct();
+					FProperty* Prop = WrapperStruct->FindPropertyByName("OptionalPropertyBagField");
+					const TSharedPtr<FJsonObject> Schema = FJsonSchemaGenerator::FPropertyToJsonSchemaObject(Prop);
+					TestBasics(*this, Schema, "object", "", FString("FInstancedPropertyBag"));
+					TestFalse("No properties field", Schema->HasField(TEXT("properties")));
+				});
+
+			It(
+				"Should generate full TOptional<FInstancedPropertyBag> schema when instance memory is provided and the optional is set.",
+				[this]() -> void
+				{
+					UStruct* WrapperStruct = FJsonUtilitiesPropertyBagWrapper::StaticStruct();
+					FProperty* Prop = WrapperStruct->FindPropertyByName("OptionalPropertyBagField");
+					TOptional<FInstancedPropertyBag> OptionalField;
+					OptionalField.Emplace();
+					OptionalField->AddProperty(FName("TestFloat"), EPropertyBagPropertyType::Float);
+					OptionalField->AddProperty(FName("TestString"), EPropertyBagPropertyType::String);
+					const TSharedPtr<FJsonObject> Schema = FJsonSchemaGenerator::FPropertyToJsonSchemaObject(
+						Prop, FJsonSchemaPropertyFilter(), nullptr, &OptionalField);
+					TestTrue("Schema", Schema.IsValid());
+					if (!Schema.IsValid())
+					{
+						return;
+					}
+					FString Type;
+					TestTrue("Type field present", Schema->TryGetStringField(TEXT("type"), Type));
+					TestEqual("Type", Type, FString("object"));
+					const TSharedPtr<FJsonObject>* Properties = nullptr;
+					TestTrue("Properties", Schema->TryGetObjectField(TEXT("properties"), Properties));
+					if (Properties)
+					{
+						const TSharedPtr<FJsonObject>* FloatProp = nullptr;
+						TestTrue("TestFloat", Properties->Get()->TryGetObjectField(TEXT("TestFloat"), FloatProp));
+						if (FloatProp)
+						{
+							TestFloat(*this, *FloatProp, "");
+						}
+						const TSharedPtr<FJsonObject>* StringProp = nullptr;
+						TestTrue("TestString", Properties->Get()->TryGetObjectField(TEXT("TestString"), StringProp));
+						if (StringProp)
+						{
+							TestString(*this, *StringProp, "");
+						}
+					}
+				});
+
+			It(
+				"Should generate a fallback TOptional<FInstancedPropertyBag> schema when instance memory is provided but the optional is not set.",
+				[this]() -> void
+				{
+					UStruct* WrapperStruct = FJsonUtilitiesPropertyBagWrapper::StaticStruct();
+					FProperty* Prop = WrapperStruct->FindPropertyByName("OptionalPropertyBagField");
+					TOptional<FInstancedPropertyBag> OptionalField; // not set
+					const TSharedPtr<FJsonObject> Schema = FJsonSchemaGenerator::FPropertyToJsonSchemaObject(
+						Prop, FJsonSchemaPropertyFilter(), nullptr, &OptionalField);
+					TestBasics(*this, Schema, "object", "", FString("FInstancedPropertyBag"));
+					TestFalse("No properties field", Schema->HasField(TEXT("properties")));
+				});
+
 			It(
 				"Should support custom callbacks.",
 				[this]() -> void
diff --git a/Engine/Plugins/Tests/EditorTests/Source/EditorTests/Private/UnrealEd/JsonUtilitiesTests.h b/Engine/Plugins/Tests/EditorTests/Source/EditorTests/Private/UnrealEd/JsonUtilitiesTests.h
--- a/Engine/Plugins/Tests/EditorTests/Source/EditorTests/Private/UnrealEd/JsonUtilitiesTests.h
+++ b/Engine/Plugins/Tests/EditorTests/Source/EditorTests/Private/UnrealEd/JsonUtilitiesTests.h
@@ -9,6 +9,8 @@
 #include "UObject/SoftObjectPath.h"
 #include "Math/Color.h"
 #include "Math/MathFwd.h"
+#include "StructUtils/InstancedStruct.h"
+#include "StructUtils/PropertyBag.h"
 
 #include "JsonUtilitiesTests.generated.h"
 
@@ -133,6 +135,51 @@ struct FJsonUtilitiesFakeStruct
 	TMap<FString, float> TestMap;
 };
 
+/**
+ * A simple inner struct used as the concrete type inside FInstancedStruct tests.
+ */
+USTRUCT()
+struct FJsonUtilitiesInstancedInnerStruct
+{
+	GENERATED_BODY()
+
+	UPROPERTY()
+	float FloatValue = 42.0f;
+
+	UPROPERTY()
+	FString StringValue = TEXT("hello");
+};
+
+/**
+ * Wrapper struct that holds an FInstancedStruct for converter and schema tests.
+ */
+USTRUCT()
+struct FJsonUtilitiesInstancedStructWrapper
+{
+	GENERATED_BODY()
+
+	UPROPERTY()
+	FInstancedStruct InstancedStructField;
+
+	UPROPERTY()
+	TOptional<FInstancedStruct> OptionalInstancedStructField;
+};
+
+/**
+ * Wrapper struct that holds an FInstancedPropertyBag for converter and schema tests.
+ */
+USTRUCT()
+struct FJsonUtilitiesPropertyBagWrapper
+{
+	GENERATED_BODY()
+
+	UPROPERTY()
+	FInstancedPropertyBag PropertyBagField;
+
+	UPROPERTY()
+	TOptional<FInstancedPropertyBag> OptionalPropertyBagField;
+};
+
 /**
  * TestClassDesc
  */
diff --git a/Engine/Source/Editor/JsonUtilitiesEditor/Private/JsonSchema/JsonSchemaGeneratorEditor.cpp b/Engine/Source/Editor/JsonUtilitiesEditor/Private/JsonSchema/JsonSchemaGeneratorEditor.cpp
--- a/Engine/Source/Editor/JsonUtilitiesEditor/Private/JsonSchema/JsonSchemaGeneratorEditor.cpp
+++ b/Engine/Source/Editor/JsonUtilitiesEditor/Private/JsonSchema/JsonSchemaGeneratorEditor.cpp
@@ -13,6 +13,8 @@
 #include "Misc/ScopeExit.h"
 #include "JsonSchema/JsonSchemaVisitorStack.h"
 #include "JsonSchema/JsonSchemaGenerator.h"
+#include "StructUtils/InstancedStruct.h"
+#include "StructUtils/PropertyBag.h"
 
 using namespace UE::JsonSchema;
 
@@ -441,14 +443,15 @@ public:
 	}
 		
 	void VisitRoot(const FVisitorTarget& VisitorTarget, 
-		const FJsonSchemaMemberPath& MemberPath, const FJsonSchemaPropertyFilter& PropertyFilter)
+		const FJsonSchemaMemberPath& MemberPath, const FJsonSchemaPropertyFilter& PropertyFilter,
+		const void* InstanceMemory)
 	{
 		const EVisitorStackElementFlags Flags = (VisitorTarget.IsType<FValidConstUStruct>() ? 
 			EVisitorStackElementFlags::IsRootStruct : EVisitorStackElementFlags::None);
     
 		VisitorStack.Empty(); // Precaution.
 
-		const FVisitorStackElement RootStackElement(VisitorTarget, Flags, MemberPath, PropertyFilter, nullptr); 
+		const FVisitorStackElement RootStackElement(VisitorTarget, Flags, MemberPath, PropertyFilter, nullptr, InstanceMemory);
 		VisitorStack.Add(RootStackElement);
 		
 		while (!VisitorStack.IsEmpty())
@@ -457,21 +460,21 @@ public:
 			
 			if (StackElement.VisitorTarget.IsType<FValidConstUStruct>())
 			{
-				VisitUStruct(StackElement.VisitorTarget.Get<FValidConstUStruct>(), 
-					StackElement.MemberPath, StackElement.PropertyFilter, 
-					StackElement.Flags);
+				VisitUStruct(StackElement.VisitorTarget.Get<FValidConstUStruct>(),
+					StackElement.MemberPath, StackElement.PropertyFilter,
+					StackElement.Flags, StackElement.InstanceMemory);
 			}
 			else if (StackElement.VisitorTarget.IsType<FValidConstFProperty>())
 			{
 				VisitFProperty(StackElement.VisitorTarget.Get<FValidConstFProperty>(),
-					StackElement.MemberPath, StackElement.PropertyFilter, StackElement.Flags);
+					StackElement.MemberPath, StackElement.PropertyFilter, StackElement.Flags, StackElement.InstanceMemory);
 			}
 		}
 	}
 
 	void VisitFProperty(const FValidConstFProperty Property, 
 		const FJsonSchemaMemberPath& MemberPath, const FJsonSchemaPropertyFilter& PropertyFilter,
-		const EVisitorStackElementFlags Flags)
+		const EVisitorStackElementFlags Flags, const void* InstanceMemory)
 	{
 		OutputEditorMetadata.CurrentPropertyMemberPath = MemberPath;
 		
@@ -529,21 +532,61 @@ public:
 		}
 
 		// Recurse into inner properties.
-		if (const UStruct* InnerStruct = GetFPropertyInnerStruct(*Property, *OutputEditorMetadata.CurrentPropertyMemberPath))
-		{
-			// See all NOTE_JSON_SCHEMA_CHECK_FLAGS_NO_PARMS.
-			FJsonSchemaPropertyFilter NewPropertyFilter = PropertyFilter;
-			NewPropertyFilter.CheckFlags &= ~CPF_ParmFlags;
+		// See all NOTE_JSON_SCHEMA_CHECK_FLAGS_NO_PARMS.
+		FJsonSchemaPropertyFilter NewPropertyFilter = PropertyFilter;
+		NewPropertyFilter.CheckFlags &= ~CPF_ParmFlags;
 
+		if (const FStructProperty* StructProp = CastField<FStructProperty>(Property))
+		{
+			// FInstancedStruct and FInstancedPropertyBag require runtime dispatch: the actual struct
+			// type and its memory are resolved from the live instance. Regular structs pass through.
+			const UStruct* InnerStruct = nullptr;
+			const void* InnerInstanceMemory = nullptr;
+			ResolveStructForRecursion(StructProp->Struct, InstanceMemory, InnerStruct, InnerInstanceMemory);
+			if (InnerStruct)
+			{
+				VisitorStack.Emplace(
+					FVisitorTarget(TInPlaceType<FValidConstUStruct>(), InnerStruct),
+					EVisitorStackElementFlags::None, OutputEditorMetadata.CurrentPropertyMemberPath,
+					NewPropertyFilter, nullptr, InnerInstanceMemory);
+			}
+		}
+		else if (const FOptionalProperty* OptProp = CastField<FOptionalProperty>(Property))
+		{
+			// Unwrap the optional to recurse into its inner struct type, if any.
+			if (const FStructProperty* InnerStructProp = CastField<FStructProperty>(OptProp->GetValueProperty()))
+			{
+				// Get the value pointer only if the optional is set; nullptr otherwise (the
+				// downstream ResolveStructForRecursion will safely skip instanced types with no instance).
+				const void* OptValueMemory = InstanceMemory
+					? OptProp->GetValuePointerForReadIfSet(InstanceMemory)
+					: nullptr;
+				const UStruct* InnerStruct = nullptr;
+				const void* InnerInstanceMemory = nullptr;
+				ResolveStructForRecursion(InnerStructProp->Struct, OptValueMemory, InnerStruct, InnerInstanceMemory);
+				if (InnerStruct)
+				{
+					VisitorStack.Emplace(
+						FVisitorTarget(TInPlaceType<FValidConstUStruct>(), InnerStruct),
+						EVisitorStackElementFlags::None, OutputEditorMetadata.CurrentPropertyMemberPath,
+						NewPropertyFilter, nullptr, InnerInstanceMemory);
+				}
+			}
+		}
+		else if (const UStruct* CollectionInnerStruct = GetFPropertyInnerStruct(*Property))
+		{
+			// Array/Set/Map inner struct types. No representative instance memory is available
+			// (elements may differ at runtime), so metadata is collected statically.
 			VisitorStack.Emplace(
-				FVisitorTarget(TInPlaceType<FValidConstUStruct>(), InnerStruct), 
-				EVisitorStackElementFlags::None, OutputEditorMetadata.CurrentPropertyMemberPath, NewPropertyFilter, nullptr);
+				FVisitorTarget(TInPlaceType<FValidConstUStruct>(), CollectionInnerStruct),
+				EVisitorStackElementFlags::None, OutputEditorMetadata.CurrentPropertyMemberPath,
+				NewPropertyFilter, nullptr, nullptr);
 		}
 	}
 	
 	void VisitUStruct(const FValidConstUStruct Struct, 
 		const FJsonSchemaMemberPath& MemberPath, const FJsonSchemaPropertyFilter& PropertyFilter,
-		const EVisitorStackElementFlags Flags)
+		const EVisitorStackElementFlags Flags, const void* InstanceMemory)
 	{
 		OutputEditorMetadata.CurrentPropertyMemberPath = MemberPath;
 		
@@ -567,18 +610,70 @@ public:
 		{
 			const FString PropertyName = Property->GetAuthoredName();
 
-			UE_JSON_SCHEMA_SCOPED_MEMBER_PATH_PUSH(OutputEditorMetadata.CurrentPropertyMemberPath, PropertyName, 
+			UE_JSON_SCHEMA_SCOPED_MEMBER_PATH_PUSH(OutputEditorMetadata.CurrentPropertyMemberPath, PropertyName,
 				!EnumHasAnyFlags(Flags, EVisitorStackElementFlags::IsRootStruct));
-			
+
+			const void* PropertyInstanceMemory = InstanceMemory
+				? static_cast<const uint8*>(InstanceMemory) + Property->GetOffset_ForInternal()
+				: nullptr;
+
 			VisitorStack.Emplace(
-				FVisitorTarget(TInPlaceType<FValidConstFProperty>(), Property), 
-				EVisitorStackElementFlags::None, OutputEditorMetadata.CurrentPropertyMemberPath, PropertyFilter, nullptr);
+				FVisitorTarget(TInPlaceType<FValidConstFProperty>(), Property),
+				EVisitorStackElementFlags::None, OutputEditorMetadata.CurrentPropertyMemberPath, PropertyFilter, nullptr, PropertyInstanceMemory);
 		}
 	}
 	
 private:
 
-	UStruct* GetFPropertyInnerStruct(const FProperty& Property, const TCHAR* PropertyPath) const
+	// Resolves the inner struct type and instance memory to use for metadata recursion.
+	// For FInstancedStruct and FInstancedPropertyBag, the actual struct type is only known at
+	// runtime; OutStruct is set to nullptr (skip recursion) when no valid instance is available.
+	// For all other structs, OutStruct is the static struct type and OutInstanceMemory is passed through.
+	static void ResolveStructForRecursion(const UScriptStruct* StaticStruct, const void* InstanceMemory,
+		const UStruct*& OutStruct, const void*& OutInstanceMemory)
+	{
+		OutStruct = nullptr;
+		OutInstanceMemory = nullptr;
+
+		if (StaticStruct == FInstancedStruct::StaticStruct())
+		{
+			if (InstanceMemory)
+			{
+				const FInstancedStruct* InstancedStruct = static_cast<const FInstancedStruct*>(InstanceMemory);
+				if (InstancedStruct->IsValid())
+				{
+					OutStruct = InstancedStruct->GetScriptStruct();
+					OutInstanceMemory = InstancedStruct->GetMemory();
+				}
+			}
+			return;
+		}
+
+		if (StaticStruct == FInstancedPropertyBag::StaticStruct())
+		{
+			if (InstanceMemory)
+			{
+				const FInstancedPropertyBag* Bag = static_cast<const FInstancedPropertyBag*>(InstanceMemory);
+				if (Bag->IsValid())
+				{
+					if (const UPropertyBag* BagStruct = Bag->GetPropertyBagStruct())
+					{
+						OutStruct = BagStruct;
+						OutInstanceMemory = Bag->GetValue().GetMemory();
+					}
+				}
+			}
+			return;
+		}
+
+		OutStruct = StaticStruct;
+		OutInstanceMemory = InstanceMemory;
+	}
+
+	// Returns the inner struct type for collection properties (Array/Set/Map), or nullptr.
+	// FStructProperty and FOptionalProperty are handled directly in VisitFProperty via
+	// ResolveStructForRecursion so that instance memory can be propagated correctly.
+	const UStruct* GetFPropertyInnerStruct(const FProperty& Property) const
 	{
 		if (const FArrayProperty* ArrayProperty = CastField<FArrayProperty>(&Property))
 		{
@@ -587,9 +682,9 @@ private:
 				return InnerStructProperty->Struct;
 			}
 		}
-		else if (const FSetProperty* SetProperty = CastField<FSetProperty>(&Property) )
+		else if (const FSetProperty* SetProperty = CastField<FSetProperty>(&Property))
 		{
-			if (const FStructProperty *ElementStructProperty = CastField<FStructProperty>(SetProperty->ElementProp))
+			if (const FStructProperty* ElementStructProperty = CastField<FStructProperty>(SetProperty->ElementProp))
 			{
 				return ElementStructProperty->Struct;
 			}
@@ -604,12 +699,6 @@ private:
 				}
 			}
 		}
-		else if (const FStructProperty* StructProperty = CastField<FStructProperty>(&Property))
-		{
-			return StructProperty->Struct;
-		}
-		
-		// @todo Push into FObjectProperty via struct instance?
 		return nullptr;
 	}
 	
@@ -622,7 +711,7 @@ private:
 //
 
 /*static*/ FJsonSchemaEditorMetadata FJsonSchemaGeneratorEditor::FPropertyToJsonSchemaMetadata(TNotNull<const FProperty*> Property,
-	const FJsonSchemaPropertyFilter& PropertyFilter)
+	const FJsonSchemaPropertyFilter& PropertyFilter, const void* InstanceMemory)
 {
 	FJsonSchemaEditorMetadata EditorMetadata;
 	
@@ -630,13 +719,14 @@ private:
 	JsonSchemaMetadataVisitor.VisitRoot(
 		FVisitorTarget(TInPlaceType<FValidConstFProperty>(), Property), 
 		FJsonSchemaMemberPath(*Property->GetAuthoredName()),
-		PropertyFilter);
+		PropertyFilter,
+		InstanceMemory);
 	
 	return EditorMetadata;
 }
 
 /*static*/ FJsonSchemaEditorMetadata FJsonSchemaGeneratorEditor::UStructToJsonSchemaMetadata(TNotNull<const UStruct*> Struct,
-	const FJsonSchemaPropertyFilter& PropertyFilter)
+	const FJsonSchemaPropertyFilter& PropertyFilter, const void* InstanceMemory)
 {
 	FJsonSchemaEditorMetadata EditorMetadata;
 	
@@ -644,7 +734,8 @@ private:
 	JsonSchemaMetadataVisitor.VisitRoot(
 		FVisitorTarget(TInPlaceType<FValidConstUStruct>(), Struct), 
 		FJsonSchemaMemberPath(),
-		PropertyFilter);
+		PropertyFilter,
+		InstanceMemory);
 	
 	return EditorMetadata;
 }
diff --git a/Engine/Source/Editor/JsonUtilitiesEditor/Private/JsonUtilitiesEditorModularFeature.cpp b/Engine/Source/Editor/JsonUtilitiesEditor/Private/JsonUtilitiesEditorModularFeature.cpp
--- a/Engine/Source/Editor/JsonUtilitiesEditor/Private/JsonUtilitiesEditorModularFeature.cpp
+++ b/Engine/Source/Editor/JsonUtilitiesEditor/Private/JsonUtilitiesEditorModularFeature.cpp
@@ -5,14 +5,14 @@
 #include "JsonSchema/JsonSchemaGeneratorEditor.h"
 
 /*virtual*/ FJsonSchemaEditorMetadata FJsonUtilitiesEditorModularFeature::FPropertyToJsonSchemaMetadata(TNotNull<const FProperty*> Property,
-	const FJsonSchemaPropertyFilter& PropertyFilter) const
+	const FJsonSchemaPropertyFilter& PropertyFilter, const void* InstanceMemory) const
 {
-	return FJsonSchemaGeneratorEditor::FPropertyToJsonSchemaMetadata(Property, PropertyFilter);
+	return FJsonSchemaGeneratorEditor::FPropertyToJsonSchemaMetadata(Property, PropertyFilter, InstanceMemory);
 }
 
 /*virtual*/ FJsonSchemaEditorMetadata FJsonUtilitiesEditorModularFeature::UStructToJsonSchemaMetadata(TNotNull<const UStruct*> Struct,
-	const FJsonSchemaPropertyFilter& PropertyFilter) const
+	const FJsonSchemaPropertyFilter& PropertyFilter, const void* InstanceMemory) const
 {
-	return FJsonSchemaGeneratorEditor::UStructToJsonSchemaMetadata(Struct, PropertyFilter);
+	return FJsonSchemaGeneratorEditor::UStructToJsonSchemaMetadata(Struct, PropertyFilter, InstanceMemory);
 }
 
diff --git a/Engine/Source/Editor/JsonUtilitiesEditor/Private/JsonUtilitiesEditorModularFeature.h b/Engine/Source/Editor/JsonUtilitiesEditor/Private/JsonUtilitiesEditorModularFeature.h
--- a/Engine/Source/Editor/JsonUtilitiesEditor/Private/JsonUtilitiesEditorModularFeature.h
+++ b/Engine/Source/Editor/JsonUtilitiesEditor/Private/JsonUtilitiesEditorModularFeature.h
@@ -13,8 +13,8 @@ class FJsonUtilitiesEditorModularFeature : public IJsonUtilitiesModularFeature
 public:
 
 	JSONUTILITIESEDITOR_API virtual FJsonSchemaEditorMetadata FPropertyToJsonSchemaMetadata(TNotNull<const FProperty*> Property,
-		const FJsonSchemaPropertyFilter& PropertyFilter = FJsonSchemaPropertyFilter()) const override;
+		const FJsonSchemaPropertyFilter& PropertyFilter, const void* InstanceMemory) const override;
 
 	JSONUTILITIESEDITOR_API virtual	FJsonSchemaEditorMetadata UStructToJsonSchemaMetadata(TNotNull<const UStruct*>,
-		const FJsonSchemaPropertyFilter& PropertyFilter = FJsonSchemaPropertyFilter()) const override;
+		const FJsonSchemaPropertyFilter& PropertyFilter, const void* InstanceMemory) const override;
 };
diff --git a/Engine/Source/Editor/JsonUtilitiesEditor/Public/JsonSchema/JsonSchemaGeneratorEditor.h b/Engine/Source/Editor/JsonUtilitiesEditor/Public/JsonSchema/JsonSchemaGeneratorEditor.h
--- a/Engine/Source/Editor/JsonUtilitiesEditor/Public/JsonSchema/JsonSchemaGeneratorEditor.h
+++ b/Engine/Source/Editor/JsonUtilitiesEditor/Public/JsonSchema/JsonSchemaGeneratorEditor.h
@@ -18,10 +18,11 @@ public:
 	 * 
 	 * @param Property The property to collect metadata for.
 	 * @param PropertyFilter Determines how to filter collected struct and property data.
+	 * @param InstanceMemory When set allows specs to be generated from FInstancedStruct and FInstancedPropertyBag.
 	 * @return The collected property metadata.
 	 */
 	JSONUTILITIESEDITOR_API static FJsonSchemaEditorMetadata FPropertyToJsonSchemaMetadata(TNotNull<const FProperty*> Property,
-		const FJsonSchemaPropertyFilter& PropertyFilter = FJsonSchemaPropertyFilter());
+		const FJsonSchemaPropertyFilter& PropertyFilter, const void* InstanceMemory);
 
 	/**
 	 * Usually not called by a developer directly. Used by the Runtime module via the Module Feature.
@@ -75,8 +76,9 @@ public:
 	 *
 	 * @param Struct The type to collect property data for. Can be either UScriptStruct, UClass or UFunction.
 	 * @param PropertyFilter Determines how to filter collected struct and property data.
+	 * @param InstanceMemory When set allows specs to be generated from FInstancedStruct and FInstancedPropertyBag.
 	 * @return The collected property metadata.
 	 */
 	JSONUTILITIESEDITOR_API static FJsonSchemaEditorMetadata UStructToJsonSchemaMetadata(TNotNull<const UStruct*> Struct,
-		const FJsonSchemaPropertyFilter& PropertyFilter = FJsonSchemaPropertyFilter());
+		const FJsonSchemaPropertyFilter& PropertyFilter, const void* InstanceMemory);
 };
diff --git a/Engine/Source/Runtime/JsonUtilities/Private/JsonObjectConverter.cpp b/Engine/Source/Runtime/JsonUtilities/Private/JsonObjectConverter.cpp
--- a/Engine/Source/Runtime/JsonUtilities/Private/JsonObjectConverter.cpp
+++ b/Engine/Source/Runtime/JsonUtilities/Private/JsonObjectConverter.cpp
@@ -18,6 +18,8 @@
 #include "UObject/VerseStringProperty.h"
 #include "Policies/CondensedJsonPrintPolicy.h"
 #include "VerseVM/VVMNames.h"
+#include "StructUtils/InstancedStruct.h"
+#include "StructUtils/PropertyBag.h"
 
 #define LOCTEXT_NAMESPACE "JsonObjectConverter"
 
@@ -258,6 +260,45 @@ TSharedPtr<FJsonValue> ConvertScalarFPropertyToJsonValueWithContainer(FProperty*
 				}
 			}
 
+			// FInstancedStruct: serialize the inner struct's fields as a flat object.
+			// Import assumes the struct is already initialized with the correct type.
+			if (StructProperty->Struct == FInstancedStruct::StaticStruct())
+			{
+				const FInstancedStruct* InstancedStruct = static_cast<const FInstancedStruct*>(Value);
+				if (!InstancedStruct->IsValid())
+				{
+					return MakeShared<FJsonValueObject>(MakeShared<FJsonObject>());
+				}
+				TSharedRef<FJsonObject> Out = MakeShared<FJsonObject>();
+				if (UStructToJsonAttributesWithContainer(InstancedStruct->GetScriptStruct(), InstancedStruct->GetMemory(), Container, ExportedObjects, Out, CheckFlags & (~CPF_ParmFlags), SkipFlags, ExportCb, ConversionFlags))
+				{
+					return MakeShared<FJsonValueObject>(Out);
+				}
+				return TSharedPtr<FJsonValue>();
+			}
+
+			// FInstancedPropertyBag: serialize the contained property values as a flat object.
+			// Import assumes the bag is already initialized with the correct schema.
+			if (StructProperty->Struct == FInstancedPropertyBag::StaticStruct())
+			{
+				const FInstancedPropertyBag* Bag = static_cast<const FInstancedPropertyBag*>(Value);
+				if (!Bag->IsValid())
+				{
+					return MakeShared<FJsonValueObject>(MakeShared<FJsonObject>());
+				}
+				const UPropertyBag* BagStruct = Bag->GetPropertyBagStruct();
+				if (!BagStruct)
+				{
+					return MakeShared<FJsonValueObject>(MakeShared<FJsonObject>());
+				}
+				TSharedRef<FJsonObject> Out = MakeShared<FJsonObject>();
+				if (UStructToJsonAttributesWithContainer(BagStruct, Bag->GetValue().GetMemory(), Container, ExportedObjects, Out, CheckFlags & (~CPF_ParmFlags), SkipFlags, ExportCb, ConversionFlags))
+				{
+					return MakeShared<FJsonValueObject>(Out);
+				}
+				return TSharedPtr<FJsonValue>();
+			}
+
 			// Default converter behavior is to use ExportTextItem if the struct has that function.
 			UScriptStruct::ICppStructOps* TheCppStructOps = StructProperty->Struct->GetCppStructOps();
 			if (TheCppStructOps && TheCppStructOps->HasExportTextItem())
@@ -859,7 +900,41 @@ namespace
 
 				if (bUseDefaultConverter)
 				{
-					if (!JsonAttributesToUStructWithContainer(Obj->Values, StructProperty->Struct, OutValue, ContainerStruct, Container, CheckFlags & (~CPF_ParmFlags), SkipFlags, bStrictMode, OutFailReason, ImportCb))
+					// This is the struct that will be used for output.
+					const UScriptStruct* ScriptStruct = nullptr;
+
+					// FInstancedStruct: applies the JSON fields to the struct type already stored in the instance.
+					// The struct must already be initialized with the correct type before importing.
+					if (StructProperty->Struct == FInstancedStruct::StaticStruct())
+					{
+						FInstancedStruct* InstancedStruct = static_cast<FInstancedStruct*>(OutValue);
+						if (InstancedStruct->IsValid())
+						{
+							ScriptStruct = InstancedStruct->GetScriptStruct();
+							OutValue = InstancedStruct->GetMutableMemory();
+						}
+					}
+					// FInstancedPropertyBag: flat property values; bag must already be initialized with the correct schema.
+					else if (StructProperty->Struct == FInstancedPropertyBag::StaticStruct())
+					{
+						FInstancedPropertyBag* Bag = static_cast<FInstancedPropertyBag*>(OutValue);
+						if (Bag->IsValid())
+						{
+							const UPropertyBag* BagStruct = Bag->GetPropertyBagStruct();
+							if (BagStruct)
+							{
+								ScriptStruct = BagStruct;
+								OutValue = Bag->GetMutableValue().GetMemory();
+							}
+						}
+					}
+					// Traditional UStructs take this path.
+					else
+					{
+						ScriptStruct = StructProperty->Struct;
+					}
+
+					if (!ScriptStruct || !JsonAttributesToUStructWithContainer(Obj->Values, ScriptStruct, OutValue, ContainerStruct, Container, CheckFlags & (~CPF_ParmFlags), SkipFlags, bStrictMode, OutFailReason, ImportCb))
 					{
 						UE_LOG(LogJson, Error, TEXT("JsonValueToUProperty - Unable to import JSON object into %s property %s"), *StructProperty->Struct->GetAuthoredName(), *Property->GetAuthoredName());
 						if (OutFailReason)
diff --git a/Engine/Source/Runtime/JsonUtilities/Private/JsonSchema/JsonSchemaGenerator.cpp b/Engine/Source/Runtime/JsonUtilities/Private/JsonSchema/JsonSchemaGenerator.cpp
--- a/Engine/Source/Runtime/JsonUtilities/Private/JsonSchema/JsonSchemaGenerator.cpp
+++ b/Engine/Source/Runtime/JsonUtilities/Private/JsonSchema/JsonSchemaGenerator.cpp
@@ -14,6 +14,8 @@
 #include "JsonUtilitiesModularFeature.h"
 #include "UObject/FieldPathProperty.h"
 #include "JsonSchema/JsonSchemaVisitorStack.h"
+#include "StructUtils/InstancedStruct.h"
+#include "StructUtils/PropertyBag.h"
 
 using namespace UE::JsonSchema;
 
@@ -56,47 +58,48 @@ public:
 	{
 	}
 	
-	TSharedPtr<FJsonObject> VisitRoot(const FVisitorTarget& VisitorTarget, 
-		const FJsonSchemaMemberPath& MemberPath, const FJsonSchemaPropertyFilter& PropertyFilter)
+	TSharedPtr<FJsonObject> VisitRoot(const FVisitorTarget& VisitorTarget,
+		const FJsonSchemaMemberPath& MemberPath, const FJsonSchemaPropertyFilter& PropertyFilter,
+		const void* InstanceMemory)
 	{
-		const EVisitorStackElementFlags Flags = (VisitorTarget.IsType<FValidConstUStruct>() ? 
+		const EVisitorStackElementFlags Flags = (VisitorTarget.IsType<FValidConstUStruct>() ?
 			EVisitorStackElementFlags::IsRootStruct : EVisitorStackElementFlags::None);
-		
+
 		VisitorStack.Empty(); // Precaution.
-		
-		const FVisitorStackElement RootStackElement(VisitorTarget, Flags, MemberPath, PropertyFilter, MakeShared<FJsonObject>()); 
+
+		const FVisitorStackElement RootStackElement(VisitorTarget, Flags, MemberPath, PropertyFilter, MakeShared<FJsonObject>(), InstanceMemory);
 		VisitorStack.Add(RootStackElement);
-		
+
 		while (!VisitorStack.IsEmpty())
 		{
 			const FVisitorStackElement StackElement = VisitorStack.Pop();
-			
+
 			if (StackElement.VisitorTarget.IsType<FValidConstUStruct>())
 			{
-				if (!VisitUStruct(StackElement.VisitorTarget.Get<FValidConstUStruct>(), 
-					StackElement.OutputSchema.ToSharedRef(), StackElement.MemberPath, StackElement.PropertyFilter, 
-					StackElement.Flags))
+				if (!VisitUStruct(StackElement.VisitorTarget.Get<FValidConstUStruct>(),
+					StackElement.OutputSchema.ToSharedRef(), StackElement.MemberPath, StackElement.PropertyFilter,
+					StackElement.Flags, StackElement.InstanceMemory))
 				{
 					return nullptr;
 				}
 			}
 			else if (StackElement.VisitorTarget.IsType<FValidConstFProperty>())
 			{
-				if (!VisitFProperty(StackElement.VisitorTarget.Get<FValidConstFProperty>(), 
+				if (!VisitFProperty(StackElement.VisitorTarget.Get<FValidConstFProperty>(),
 					StackElement.OutputSchema.ToSharedRef(), StackElement.MemberPath, StackElement.PropertyFilter,
-					StackElement.Flags))
+					StackElement.Flags, StackElement.InstanceMemory))
 				{
 					return nullptr;
 				}
 			}
 		}
-	
+
 		return RootStackElement.OutputSchema;
 	}
 	
-	bool VisitFProperty(const FValidConstFProperty Property, const TSharedRef<FJsonObject>& OutputSchema, 
+	bool VisitFProperty(const FValidConstFProperty Property, const TSharedRef<FJsonObject>& OutputSchema,
 		const FJsonSchemaMemberPath& MemberPath, const FJsonSchemaPropertyFilter& PropertyFilter,
-		const EVisitorStackElementFlags Flags)
+		const EVisitorStackElementFlags Flags, const void* InstanceMemory)
 	{
 		EditorMetadata.CurrentPropertyMemberPath = MemberPath;
 		
@@ -112,11 +115,11 @@ public:
 		}
 		else if (const FOptionalProperty* OptionalProperty = CastField<FOptionalProperty>(Property))
 		{
-			ValidProperty = BuildJsonForOptionalProperty(*OptionalProperty, PropertyFilter, OutputSchema);
+			ValidProperty = BuildJsonForOptionalProperty(*OptionalProperty, PropertyFilter, OutputSchema, InstanceMemory);
 		}
 		else if (const FStructProperty* StructProperty = CastField<FStructProperty>(Property))
 		{
-			ValidProperty = BuildJsonForStructProperty(*StructProperty, PropertyFilter, OutputSchema);
+			ValidProperty = BuildJsonForStructProperty(*StructProperty, PropertyFilter, OutputSchema, InstanceMemory);
 		}
 		else if (const FObjectProperty* ObjectProperty = CastField<FObjectProperty>(Property))
 		{
@@ -196,9 +199,9 @@ public:
 		return true;
 	}
 	
-	bool VisitUStruct(const FValidConstUStruct Struct, const TSharedRef<FJsonObject>& OutputSchema, 
+	bool VisitUStruct(const FValidConstUStruct Struct, const TSharedRef<FJsonObject>& OutputSchema,
 		const FJsonSchemaMemberPath& MemberPath, const FJsonSchemaPropertyFilter& PropertyFilter,
-		const EVisitorStackElementFlags Flags)
+		const EVisitorStackElementFlags Flags, const void* InstanceMemory)
 	{
 		EditorMetadata.CurrentPropertyMemberPath = MemberPath;
 
@@ -229,7 +232,7 @@ public:
 
 		// If property is to be included - Adds property to JSON and possibly the required list, and also adds to the visitor stack for processing.
 		// Returns whether the property was included.
-		auto MaybeIncludeProperty = [this, PropertyFilter, Flags, bUseJsonRpcFormat](const FProperty* Property,
+		auto MaybeIncludeProperty = [this, PropertyFilter, Flags, bUseJsonRpcFormat, InstanceMemory](const FProperty* Property,
 			const TSharedRef<FJsonObject>& PropertiesSchema, TArray<TSharedPtr<FJsonValue>>& RequiredPropertiesArray) -> bool
 			{
 				const FString PropertyName = Property->GetAuthoredName();
@@ -252,12 +255,17 @@ public:
 				const FString JsonMemberName = PropertyFilter.PropertyAuthoredNameToJsonKey(PropertyName); 
 				PropertiesSchema->SetObjectField(JsonMemberName, PropertySchema);
 				
+				const void* PropertyInstanceMemory = InstanceMemory
+					? static_cast<const uint8*>(InstanceMemory) + Property->GetOffset_ForInternal()
+					: nullptr;
+
 				VisitorStack.Emplace(
-					FVisitorTarget(TInPlaceType<FValidConstFProperty>(), Property), 
+					FVisitorTarget(TInPlaceType<FValidConstFProperty>(), Property),
 					bUseJsonRpcFormat ? EVisitorStackElementFlags::RequireAllProperties : EVisitorStackElementFlags::None,
-					EditorMetadata.CurrentPropertyMemberPath, 
+					EditorMetadata.CurrentPropertyMemberPath,
 					NewPropertyFilter,
-					PropertySchema);
+					PropertySchema,
+					PropertyInstanceMemory);
 
 				const TSharedPtr<FJsonObject> PropertyMetadata = EditorMetadata.GetPropertyMetadataForCurrentPropertyMemberPath();
 
@@ -393,57 +401,127 @@ private:
 		return true;
 	}
 
-	bool BuildJsonForOptionalProperty(const FOptionalProperty& OptionalProperty, 
-		const FJsonSchemaPropertyFilter& PropertyFilter, const TSharedRef<FJsonObject>& OutputSchema)
+	bool BuildJsonForOptionalProperty(const FOptionalProperty& OptionalProperty,
+		const FJsonSchemaPropertyFilter& PropertyFilter, const TSharedRef<FJsonObject>& OutputSchema,
+		const void* InstanceMemory)
 	{
+		// If the optional is set, pass the inner value pointer so that runtime-typed properties
+		// (e.g. FInstancedStruct) inside a TOptional can produce an accurate schema.
+		// GetValuePointerForReadIfSet returns nullptr when the optional is unset, which is the
+		// correct fallback (the inner property visitor will use a static/type-only schema).
+		const void* InnerInstanceMemory = InstanceMemory
+			? OptionalProperty.GetValuePointerForReadIfSet(InstanceMemory)
+			: nullptr;
+
 		VisitorStack.Emplace(
-			FVisitorTarget(TInPlaceType<FValidConstFProperty>(), OptionalProperty.GetValueProperty()), 
-			EVisitorStackElementFlags::None, 
-			EditorMetadata.CurrentPropertyMemberPath, 
+			FVisitorTarget(TInPlaceType<FValidConstFProperty>(), OptionalProperty.GetValueProperty()),
+			EVisitorStackElementFlags::None,
+			EditorMetadata.CurrentPropertyMemberPath,
 			PropertyFilter,
-			OutputSchema);
+			OutputSchema,
+			InnerInstanceMemory);
 		return true;
 	}
 
-	bool BuildJsonForStructProperty(const FStructProperty& StructProperty, 
-		const FJsonSchemaPropertyFilter& PropertyFilter, const TSharedRef<FJsonObject>& OutputSchema)
+	bool BuildJsonForStructProperty(const FStructProperty& StructProperty,
+		const FJsonSchemaPropertyFilter& PropertyFilter, const TSharedRef<FJsonObject>& OutputSchema,
+		const void* InstanceMemory)
 	{
+		if (StructProperty.Struct == FInstancedStruct::StaticStruct())
+		{
+			// With instance memory, push the actual struct so VisitUStruct generates the full schema.
+			if (InstanceMemory)
+			{
+				const FInstancedStruct* InstancedStruct = static_cast<const FInstancedStruct*>(InstanceMemory);
+				if (InstancedStruct->IsValid())
+				{
+					VisitorStack.Emplace(
+						FVisitorTarget(TInPlaceType<FValidConstUStruct>(), InstancedStruct->GetScriptStruct()),
+						EVisitorStackElementFlags::None,
+						EditorMetadata.CurrentPropertyMemberPath,
+						PropertyFilter,
+						OutputSchema,
+						InstancedStruct->GetMemory());
+					return true;
+				}
+			}
+			// No instance or invalid struct: static fallback.
+			OutputSchema->SetStringField(TEXT("type"), TEXT("object"));
+			OutputSchema->SetStringField(TEXT("title"), TEXT("FInstancedStruct"));
+			return true;
+		}
+
+		if (StructProperty.Struct == FInstancedPropertyBag::StaticStruct())
+		{
+			// With instance memory, push the actual UPropertyBag so VisitUStruct generates the full schema.
+			if (InstanceMemory)
+			{
+				const FInstancedPropertyBag* Bag = static_cast<const FInstancedPropertyBag*>(InstanceMemory);
+				if (Bag->IsValid())
+				{
+					if (const UPropertyBag* BagStruct = Bag->GetPropertyBagStruct())
+					{
+						VisitorStack.Emplace(
+							FVisitorTarget(TInPlaceType<FValidConstUStruct>(), BagStruct),
+							EVisitorStackElementFlags::None,
+							EditorMetadata.CurrentPropertyMemberPath,
+							PropertyFilter,
+							OutputSchema,
+							Bag->GetValue().GetMemory());
+						return true;
+					}
+				}
+			}
+			// No instance or invalid bag: static fallback.
+			OutputSchema->SetStringField(TEXT("type"), TEXT("object"));
+			OutputSchema->SetStringField(TEXT("title"), TEXT("FInstancedPropertyBag"));
+			return true;
+		}
+
 		VisitorStack.Emplace(
 			FVisitorTarget(TInPlaceType<FValidConstUStruct>(), StructProperty.Struct),
 			EVisitorStackElementFlags::None,
-			EditorMetadata.CurrentPropertyMemberPath, 
+			EditorMetadata.CurrentPropertyMemberPath,
 			PropertyFilter,
-			OutputSchema);
-		OutputSchema->SetStringField("title", StructProperty.Struct->GetName());
+			OutputSchema,
+			InstanceMemory);
+		OutputSchema->SetStringField(TEXT("title"), StructProperty.Struct->GetName());
 		return true;
 	}
 
-	bool BuildJsonForArrayProperty(const FArrayProperty& ArrayProperty, 
+	bool BuildJsonForArrayProperty(const FArrayProperty& ArrayProperty,
 		const FJsonSchemaPropertyFilter& PropertyFilter, const TSharedRef<FJsonObject>& OutputSchema)
 	{
 		const TSharedRef<FJsonObject> ItemsJson = MakeShared<FJsonObject>();
 		VisitorStack.Emplace(
-			FVisitorTarget(TInPlaceType<FValidConstFProperty>(), ArrayProperty.Inner), 
+			FVisitorTarget(TInPlaceType<FValidConstFProperty>(), ArrayProperty.Inner),
 			EVisitorStackElementFlags::SkipDescription,
-			EditorMetadata.CurrentPropertyMemberPath, 
+			EditorMetadata.CurrentPropertyMemberPath,
 			PropertyFilter,
-			ItemsJson);
+			ItemsJson,
+			// No instance memory: a TArray can hold elements of different runtime types (e.g. each
+			// FInstancedStruct element may have a distinct UScriptStruct), so there is no single
+			// representative instance. The element visitor falls back to a static/type-only schema.
+			nullptr);
 
 		OutputSchema->SetStringField(TEXT("type"), TEXT("array"));
 		OutputSchema->SetObjectField(TEXT("items"), ItemsJson);
 		return true;
 	}
 
-	bool BuildJsonForSetProperty(const FSetProperty& SetProperty, 
+	bool BuildJsonForSetProperty(const FSetProperty& SetProperty,
 		const FJsonSchemaPropertyFilter& PropertyFilter, const TSharedRef<FJsonObject>& OutputSchema)
 	{
 		const TSharedRef<FJsonObject> ItemsJson = MakeShared<FJsonObject>();
 		VisitorStack.Emplace(
-			FVisitorTarget(TInPlaceType<FValidConstFProperty>(), SetProperty.ElementProp), 
+			FVisitorTarget(TInPlaceType<FValidConstFProperty>(), SetProperty.ElementProp),
 			EVisitorStackElementFlags::SkipDescription,
-			EditorMetadata.CurrentPropertyMemberPath, 
+			EditorMetadata.CurrentPropertyMemberPath,
 			PropertyFilter,
-			ItemsJson);
+			ItemsJson,
+			// No instance memory: same reasoning as BuildJsonForArrayProperty - TSet has no
+			// single representative element to drive runtime-typed schema generation.
+			nullptr);
 
 		OutputSchema->SetStringField(TEXT("type"), TEXT("array"));
 		OutputSchema->SetBoolField(TEXT("uniqueItems"), true);
@@ -462,7 +540,10 @@ private:
 				EVisitorStackElementFlags::SkipDescription,
 				EditorMetadata.CurrentPropertyMemberPath, 
 				PropertyFilter,
-				MapValuesJson);
+				MapValuesJson,
+				// No instance memory: same reasoning as BuildJsonForArrayProperty - TMap has no
+				// single representative element to drive runtime-typed schema generation.
+				nullptr);
 
 			OutputSchema->SetStringField(TEXT("type"), TEXT("object"));
 			OutputSchema->SetObjectField(TEXT("additionalProperties"), MapValuesJson);
@@ -533,7 +614,8 @@ private:
 //
 
 /*static*/ TSharedPtr<FJsonObject> FJsonSchemaGenerator::FPropertyToJsonSchemaObject(TNotNull<const FProperty*> Property,
-	const FJsonSchemaPropertyFilter& PropertyFilter, const FJsonSchemaEditorMetadata* CachedEditorMetadata)
+	const FJsonSchemaPropertyFilter& PropertyFilter, const FJsonSchemaEditorMetadata* CachedEditorMetadata,
+	const void* InstanceMemory)
 {
 	// See all NOTE_JSON_UTILITIES_SCHEMA_GENERATOR_METADATA_MODULAR_FEATURE.
 	FJsonSchemaEditorMetadata EditorMetadata;
@@ -543,14 +625,15 @@ private:
 	}
 	else if (IJsonUtilitiesModularFeature* ModuleFeature = IJsonUtilitiesModularFeature::Get())
 	{
-		EditorMetadata = ModuleFeature->FPropertyToJsonSchemaMetadata(Property, PropertyFilter);
+		EditorMetadata = ModuleFeature->FPropertyToJsonSchemaMetadata(Property, PropertyFilter, InstanceMemory);
 	}
-	
+
 	FJsonSchemaGeneratorVisitor JsonSchemaGeneratorVisitor(EditorMetadata);
 	return JsonSchemaGeneratorVisitor.VisitRoot(
-		FVisitorTarget(TInPlaceType<FValidConstFProperty>(), Property), 
-		FJsonSchemaMemberPath(*Property->GetAuthoredName()),\
-		PropertyFilter);
+		FVisitorTarget(TInPlaceType<FValidConstFProperty>(), Property),
+		FJsonSchemaMemberPath(*Property->GetAuthoredName()),
+		PropertyFilter,
+		InstanceMemory);
 }
 
 /*static*/ FString FJsonSchemaGenerator::FPropertyToJsonSchemaString(TNotNull<const FProperty*> Property,
@@ -575,7 +658,8 @@ private:
 }
 
 /*static*/ TSharedPtr<FJsonObject> FJsonSchemaGenerator::UStructToJsonSchemaObject(TNotNull<const UStruct*> Struct,
-	const FJsonSchemaPropertyFilter& PropertyFilter, const FJsonSchemaEditorMetadata* CachedEditorMetadata)
+	const FJsonSchemaPropertyFilter& PropertyFilter, const FJsonSchemaEditorMetadata* CachedEditorMetadata,
+	const void* InstanceMemory)
 {
 	// See all NOTE_JSON_UTILITIES_SCHEMA_GENERATOR_METADATA_MODULAR_FEATURE.
 	FJsonSchemaEditorMetadata EditorMetadata;
@@ -585,14 +669,15 @@ private:
 	}
 	else if (IJsonUtilitiesModularFeature* ModuleFeature = IJsonUtilitiesModularFeature::Get())
 	{
-		EditorMetadata = ModuleFeature->UStructToJsonSchemaMetadata(Struct, PropertyFilter);
+		EditorMetadata = ModuleFeature->UStructToJsonSchemaMetadata(Struct, PropertyFilter, InstanceMemory);
 	}
-	
+
 	FJsonSchemaGeneratorVisitor JsonSchemaGeneratorVisitor(EditorMetadata);
 	return JsonSchemaGeneratorVisitor.VisitRoot(
-		FVisitorTarget(TInPlaceType<FValidConstUStruct>(), Struct), 
+		FVisitorTarget(TInPlaceType<FValidConstUStruct>(), Struct),
 		FJsonSchemaMemberPath(),
-		PropertyFilter);
+		PropertyFilter,
+		InstanceMemory);
 }
 
 /*static*/ FString FJsonSchemaGenerator::UStructToJsonSchemaString(TNotNull<const UStruct*> Struct,
diff --git a/Engine/Source/Runtime/JsonUtilities/Private/JsonSchema/JsonSchemaVisitorStack.h b/Engine/Source/Runtime/JsonUtilities/Private/JsonSchema/JsonSchemaVisitorStack.h
--- a/Engine/Source/Runtime/JsonUtilities/Private/JsonSchema/JsonSchemaVisitorStack.h
+++ b/Engine/Source/Runtime/JsonUtilities/Private/JsonSchema/JsonSchemaVisitorStack.h
@@ -27,28 +27,33 @@ namespace UE::JsonSchema
 
 	struct FVisitorStackElement
 	{
-		explicit FVisitorStackElement(const FVisitorTarget& InVisitorTarget, 
+		explicit FVisitorStackElement(const FVisitorTarget& InVisitorTarget,
 				const EVisitorStackElementFlags InFlags,
-				const FJsonSchemaMemberPath& InMemberPath, 
+				const FJsonSchemaMemberPath& InMemberPath,
 				const FJsonSchemaPropertyFilter& InPropertyFilter,
-				const TSharedPtr<FJsonObject>& InOutputSchema) :
+				const TSharedPtr<FJsonObject>& InOutputSchema,
+				const void* InInstanceMemory) :
 			VisitorTarget(InVisitorTarget),
 			Flags(InFlags),
 			MemberPath(InMemberPath),
 			PropertyFilter(InPropertyFilter),
-			OutputSchema(InOutputSchema)
+			OutputSchema(InOutputSchema),
+			InstanceMemory(InInstanceMemory)
 		{
 			if (VisitorTarget.IsType<FValidConstFProperty>())
 			{
 				check(!EnumHasAnyFlags(Flags, EVisitorStackElementFlags::IsRootStruct));
 			}
 		}
-		
-		FVisitorTarget VisitorTarget; 
+
+		FVisitorTarget VisitorTarget;
 		EVisitorStackElementFlags Flags = EVisitorStackElementFlags::None;
 		FJsonSchemaMemberPath MemberPath;
 		FJsonSchemaPropertyFilter PropertyFilter;
 		const TSharedPtr<FJsonObject> OutputSchema;
+		/** Optional pointer to the live instance for this target. Used to generate accurate schemas for
+		 *  types whose structure is only known at runtime (e.g. FInstancedStruct, FInstancedPropertyBag). */
+		const void* InstanceMemory = nullptr;
 	};
 
 	using FVisitorStack = TArray<FVisitorStackElement>; 
diff --git a/Engine/Source/Runtime/JsonUtilities/Private/JsonUtilitiesModularFeature.h b/Engine/Source/Runtime/JsonUtilities/Private/JsonUtilitiesModularFeature.h
--- a/Engine/Source/Runtime/JsonUtilities/Private/JsonUtilitiesModularFeature.h
+++ b/Engine/Source/Runtime/JsonUtilities/Private/JsonUtilitiesModularFeature.h
@@ -25,8 +25,8 @@ public:
 
 	
 	virtual FJsonSchemaEditorMetadata FPropertyToJsonSchemaMetadata(TNotNull<const FProperty*> Property,
-		const FJsonSchemaPropertyFilter& PropertyFilter = FJsonSchemaPropertyFilter()) const = 0;
+		const FJsonSchemaPropertyFilter& PropertyFilter, const void* InstanceMemory) const = 0;
 
 	virtual FJsonSchemaEditorMetadata UStructToJsonSchemaMetadata(TNotNull<const UStruct*> Struct,
-		const FJsonSchemaPropertyFilter& PropertyFilter = FJsonSchemaPropertyFilter()) const = 0;
+		const FJsonSchemaPropertyFilter& PropertyFilter, const void* InstanceMemory) const = 0;
 };
diff --git a/Engine/Source/Runtime/JsonUtilities/Public/JsonSchema/JsonSchemaGenerator.h b/Engine/Source/Runtime/JsonUtilities/Public/JsonSchema/JsonSchemaGenerator.h
--- a/Engine/Source/Runtime/JsonUtilities/Public/JsonSchema/JsonSchemaGenerator.h
+++ b/Engine/Source/Runtime/JsonUtilities/Public/JsonSchema/JsonSchemaGenerator.h
@@ -45,11 +45,13 @@ public:
 	 * @param PropertyFilter Determines how to filter collected struct and property data.
 	 * @param CachedEditorMetadata If valid, use this metadata instead of collecting it. Otherwise, it will be automatically collected
 	 *		(if running in Editor context.) This is optional, and used for workflows that have cached metadata from an earlier process.
+	 * @param InstanceMemory When set allows specs to be generated from FInstancedStruct and FInstancedPropertyBag.
 	 * @return Schema for the FProperty, as a JSON object.
 	 */
 	JSONUTILITIES_API TSharedPtr<FJsonObject> static FPropertyToJsonSchemaObject(TNotNull<const FProperty*> Property,
-		const FJsonSchemaPropertyFilter& PropertyFilter = FJsonSchemaPropertyFilter(), 
-		const FJsonSchemaEditorMetadata* CachedEditorMetadata = nullptr);
+		const FJsonSchemaPropertyFilter& PropertyFilter = FJsonSchemaPropertyFilter(),
+		const FJsonSchemaEditorMetadata* CachedEditorMetadata = nullptr,
+		const void* InstanceMemory = nullptr);
 	
 	/**
 	 * String output variant of FPropertyToJsonSchemaObject() (see above.)
@@ -144,11 +146,13 @@ public:
 	 * @param PropertyFilter Determines how to filter collected struct and property data.
 	 * @param CachedEditorMetadata If valid, use this metadata instead of collecting it. Otherwise, it will be automatically collected
 	 *		(if running in Editor context.) This is optional, and used for workflows that have cached metadata from an earlier process.
+	 * @param InstanceMemory When set allows specs to be generated from FInstancedStruct and FInstancedPropertyBag.
 	 * @return Schema for the UStruct, as a JSON object.
 	 */
 	JSONUTILITIES_API static TSharedPtr<FJsonObject> UStructToJsonSchemaObject(TNotNull<const UStruct*> Struct,
 		const FJsonSchemaPropertyFilter& PropertyFilter = FJsonSchemaPropertyFilter(),
-		const FJsonSchemaEditorMetadata* CachedEditorMetadata = nullptr);
+		const FJsonSchemaEditorMetadata* CachedEditorMetadata = nullptr,
+		const void* InstanceMemory = nullptr);
 	
 	/**
 	 * Templated variant of UStructToJsonSchemaObject() (see above.).
