[JsonUtilities] Fix in EXCEPTION_STACK_OVERFLOW in FJsonObjectConverter for circular Instanced properties

Fix: Move the cycle guards (self-reference + ExportedObjects) above the
CPF_PersistentInstance check in ShouldExportObjectPropertyByValue().
Instanced properties still export by value, but now respect cycle detection.

#Synced-CL 51884462

diff --git a/Engine/Source/Runtime/JsonUtilities/Private/JsonObjectConverter.cpp b/Engine/Source/Runtime/JsonUtilities/Private/JsonObjectConverter.cpp
--- a/Engine/Source/Runtime/JsonUtilities/Private/JsonObjectConverter.cpp
+++ b/Engine/Source/Runtime/JsonUtilities/Private/JsonObjectConverter.cpp
@@ -60,12 +60,7 @@ namespace
 			return false;
 		}
 
-		// Check the instanced flag for backwards compatibility - always export by value in this case.
-		if (Property->HasAnyPropertyFlags(CPF_PersistentInstance)
-			|| (OuterProperty && OuterProperty->HasAnyPropertyFlags(CPF_PersistentInstance)))
-		{
-			return true;
-		}
+		// Check cycle guards BEFORE the instanced flag
 
 		// Check if it's a reference to the container (self) to guard against cycles.
 		if (Value == Container)
@@ -80,6 +75,13 @@ namespace
 			return false;
 		}
 
+		// Check the instanced flag for backwards compatibility - always export by value in this case.
+		if (Property->HasAnyPropertyFlags(CPF_PersistentInstance)
+			|| (OuterProperty && OuterProperty->HasAnyPropertyFlags(CPF_PersistentInstance)))
+		{
+			return true;
+		}
+
 		// Export by value if it is scoped within the current container context (if set).
 		if (Container
 			&& Value->IsInOuter(Container))
