[Json] Fix json library to use actual object class instead of static class

#tests Used locally when serialized using UObject type on variout class types.
#Synced-CL 52304474

diff --git a/Engine/Source/Runtime/JsonUtilities/Public/JsonObjectConverter.h b/Engine/Source/Runtime/JsonUtilities/Public/JsonObjectConverter.h
--- a/Engine/Source/Runtime/JsonUtilities/Public/JsonObjectConverter.h
+++ b/Engine/Source/Runtime/JsonUtilities/Public/JsonObjectConverter.h
@@ -22,8 +22,8 @@
 #include "Templates/SharedPointer.h"
 #include "Trace/Detail/Channel.h"
 #include "UObject/Class.h"
-#include "Templates/Models.h"
 #include "Concepts/StaticClassProvider.h"
+#include "Concepts/UObject.h"
 
 #define LOCTEXT_NAMESPACE "JsonObjectConverter"
 
@@ -88,6 +88,19 @@ public: // UStruct -> JSON
 
 	static JSONUTILITIES_API const CustomExportCallback ExportCallback_WriteISO8601Dates;
 
+	template<typename InStructType>
+	static const UStruct* GetStructDefinition(TNotNull<const InStructType*> InStruct)
+	{
+		if constexpr (UE::CUObject<InStructType>)
+		{
+			return InStruct->GetClass();
+		}
+		else
+		{
+			return InStructType::StaticStruct();
+		}
+	}
+
 	/**
 	 * Templated version of UStructToJsonObject to try and make most of the params. Also serves as an example use case
 	 *
@@ -101,10 +114,12 @@ public: // UStruct -> JSON
 	static TSharedPtr<FJsonObject> UStructToJsonObject(const InStructType& InStruct, int64 CheckFlags = 0, int64 SkipFlags = 0, const CustomExportCallback* ExportCb = nullptr)
 	{
 		TSharedRef<FJsonObject> JsonObject = MakeShared<FJsonObject>();
-		if (UStructToJsonObject(InStructType::StaticStruct(), &InStruct, JsonObject, CheckFlags, SkipFlags, ExportCb))
+
+		if (UStructToJsonObject(GetStructDefinition<InStructType>(&InStruct), &InStruct, JsonObject, CheckFlags, SkipFlags, ExportCb))
 		{
 			return JsonObject;
 		}
+
 		return TSharedPtr<FJsonObject>(); // something went wrong
 	}
 
@@ -155,14 +170,7 @@ public: // UStruct -> JSON
 	template<typename InStructType>
 	static bool UStructToJsonObjectString(const InStructType& InStruct, FString& OutJsonString, int64 CheckFlags = 0, int64 SkipFlags = 0, int32 Indent = 0, const CustomExportCallback* ExportCb = nullptr, bool bPrettyPrint = true)
 	{
-		if constexpr (TModels<CStaticClassProvider, InStructType>::Value)
-		{
-			return UStructToJsonObjectString(InStructType::StaticClass(), &InStruct, OutJsonString, CheckFlags, SkipFlags, Indent, ExportCb, bPrettyPrint);
-		}
-		else
-		{
-			return UStructToJsonObjectString(InStructType::StaticStruct(), &InStruct, OutJsonString, CheckFlags, SkipFlags, Indent, ExportCb, bPrettyPrint);
-		}
+		return UStructToJsonObjectString(GetStructDefinition<InStructType>(&InStruct), &InStruct, OutJsonString, CheckFlags, SkipFlags, Indent, ExportCb, bPrettyPrint);
 	}
 
 	/**
@@ -254,14 +262,7 @@ public: // JSON -> UStruct
 	template<typename OutStructType>
 	static bool JsonObjectToUStruct(const TSharedRef<FJsonObject>& JsonObject, OutStructType* OutStruct, int64 CheckFlags = 0, int64 SkipFlags = 0, const bool bStrictMode = false, FText* OutFailReason = nullptr, const CustomImportCallback* ImportCb = nullptr)
 	{
-		if constexpr (TModels<CStaticClassProvider, OutStructType>::Value)
-		{
-			return JsonObjectToUStruct(JsonObject, OutStructType::StaticClass(), OutStruct, CheckFlags, SkipFlags, bStrictMode, OutFailReason, ImportCb);
-		}
-		else
-		{
-			return JsonObjectToUStruct(JsonObject, OutStructType::StaticStruct(), OutStruct, CheckFlags, SkipFlags, bStrictMode, OutFailReason, ImportCb);
-		}
+		return JsonObjectToUStruct(JsonObject, GetStructDefinition<OutStructType>(OutStruct), OutStruct, CheckFlags, SkipFlags, bStrictMode, OutFailReason, ImportCb);
 	}
 
 	/**
